Polished up math expression node, almost ready for prime time (left in experimental for now until we can fully test it).
(TTP #305005)

diff --git a/Engine/Plugins/Experimental/MathExpressions/Source/MathExpressions/Classes/K2Node_MathExpression.h b/Engine/Plugins/Experimental/MathExpressions/Source/MathExpressions/Classes/K2Node_MathExpression.h
--- a/Engine/Plugins/Experimental/MathExpressions/Source/MathExpressions/Classes/K2Node_MathExpression.h
+++ b/Engine/Plugins/Experimental/MathExpressions/Source/MathExpressions/Classes/K2Node_MathExpression.h
@@ -4,6 +4,13 @@
 
 #include "K2Node_MathExpression.generated.h"
 
+/**
+* This node type acts like our collapsed nodes, a single node that represents
+* a larger sub-network of nodes (contained within a sub-graph). This node will
+* take the expression it was named with, and attempt to convert it into a
+* series of math node. If it is unsuccessful, then it generates a series of
+* errors that are tacked onto the node for the user to see.
+*/
 UCLASS()
 class UK2Node_MathExpression : public UK2Node_Composite
 {
@@ -12,12 +19,9 @@ class UK2Node_MathExpression : public UK2Node_Composite
 public:
 
 	// The math expression to evaluate
-	UPROPERTY(EditAnywhere, Category=Expression)
-	FString Expression;
+	UPROPERTY(EditAnywhere, Category = Expression)
+		FString Expression;
 
-	//@TODO: Debug output - the parsing results
-	UPROPERTY(VisibleAnywhere, Category=Expression)
-	FString ParseResults;
 public:
 	// UObject interface
 	virtual void PostEditChangeProperty(struct FPropertyChangedEvent& PropertyChangedEvent) OVERRIDE;
@@ -34,10 +38,29 @@ public:
 
 	// UK2Node interface
 	virtual void GetMenuEntries(FGraphContextMenuBuilder& ContextMenuBuilder) const OVERRIDE;
+	virtual void ValidateNodeDuringCompilation(class FCompilerResultsLog& MessageLog) const OVERRIDE;
 	// End of UK2Node interface
 
-protected:
-	void RebuildExpression();
+private:
+	/**
+	* Clears this node's sub-graph, and then takes the supplied string and
+	* parses, and converts it into a series of new graph nodes.
+	*
+	* @param  NewExpression	The new string we wish to replace the current "Expression".
+	*/
+	void RebuildExpression(FString NewExpression);
+
+	/**
+	* Clears the cached expression string, deletes all generated node, clears
+	* input pins, and resets the parser and graph generator (priming this node
+	* for a new expression).
+	*/
+	void ClearExpression();
+
+private:
+	/** Cached so we don't have to regenerate it when the graph is recompiled */
+	//UPROPERTY(Transient)
+	TSharedPtr<class FCompilerResultsLog> CachedMessageLog;
 };
 
 
diff --git a/Engine/Plugins/Experimental/MathExpressions/Source/MathExpressions/Private/BasicTokenParser.cpp b/Engine/Plugins/Experimental/MathExpressions/Source/MathExpressions/Private/BasicTokenParser.cpp
new file mode 100644
--- /dev/null
+++ b/Engine/Plugins/Experimental/MathExpressions/Source/MathExpressions/Private/BasicTokenParser.cpp
@@ -0,0 +1,951 @@
+// Copyright 1998-2014 Epic Games, Inc. All Rights Reserved.
+
+#include "MathExpressionsPrivatePCH.h"
+#include "BasicTokenParser.h"
+
+#define LOCTEXT_NAMESPACE "BasicTokenParser"
+DEFINE_LOG_CATEGORY_STATIC(LogTokenParser, Log, All);
+
+/*******************************************************************************
+ * FBasicToken
+*******************************************************************************/
+
+//------------------------------------------------------------------------------
+FBasicToken::FBasicToken()
+{
+	InitToken();
+}
+
+//------------------------------------------------------------------------------
+void FBasicToken::InitToken(EPropertyType InConstType)
+{
+	ConstantType = InConstType;
+	TokenType = FBasicToken::TOKEN_None;
+	TokenName	 = NAME_None;
+	StartPos	 = 0;
+	StartLine	 = 0;
+	*Identifier	 = 0;
+	FMemory::Memzero(String, sizeof(Identifier));
+}
+
+//------------------------------------------------------------------------------
+void FBasicToken::Clone(const FBasicToken& Other)
+{
+	TokenType = Other.TokenType;
+	TokenName = Other.TokenName;
+	StartPos  = Other.StartPos;
+	StartLine = Other.StartLine;
+
+	FCString::Strncpy(Identifier, Other.Identifier, NAME_SIZE);
+	FMemory::Memcpy(String, Other.String, sizeof(String));
+}
+
+//------------------------------------------------------------------------------
+bool FBasicToken::Matches(TCHAR const* Str, ESearchCase::Type SearchCase) const
+{
+	return (TokenType==TOKEN_Identifier || TokenType==TOKEN_Symbol) && ((SearchCase == ESearchCase::CaseSensitive) ? !FCString::Strcmp(Identifier, Str) : !FCString::Stricmp(Identifier, Str));
+}
+
+//------------------------------------------------------------------------------
+bool FBasicToken::Matches(FName const& Name) const
+{
+	return TokenType==TOKEN_Identifier && TokenName==Name;
+}
+
+//------------------------------------------------------------------------------
+bool FBasicToken::StartsWith(TCHAR const* Str, bool bCaseSensitive) const
+{
+	const int32 StrLength = FCString::Strlen(Str);
+	return (TokenType==TOKEN_Identifier || TokenType==TOKEN_Symbol) && (bCaseSensitive ? (!FCString::Strncmp(Identifier, Str, StrLength)) : (!FCString::Strnicmp(Identifier, Str, StrLength)));
+}
+
+//------------------------------------------------------------------------------
+bool FBasicToken::IsBool() const
+{
+	return ConstantType == CPT_Bool || ConstantType == CPT_Bool8 || ConstantType == CPT_Bool16 || ConstantType == CPT_Bool32 || ConstantType == CPT_Bool64;
+}
+
+//------------------------------------------------------------------------------
+void FBasicToken::SetConstInt(int32 InInt)
+{
+	ConstantType = CPT_Int;
+	Int			 = InInt;
+	TokenType	 = TOKEN_Const;
+}
+
+//------------------------------------------------------------------------------
+void FBasicToken::SetConstBool(bool InBool)
+{
+	ConstantType = CPT_Bool;
+	NativeBool   = InBool;
+	TokenType    = TOKEN_Const;
+}
+
+//------------------------------------------------------------------------------
+void FBasicToken::SetConstFloat(float InFloat)
+{
+	ConstantType = CPT_Float;
+	Float        = InFloat;
+	TokenType    = TOKEN_Const;
+}
+
+//------------------------------------------------------------------------------
+void FBasicToken::SetConstName(FName InName)
+{
+	ConstantType = CPT_Name;
+	*(FName*)NameBytes = InName;
+	TokenType = TOKEN_Const;
+}
+
+//------------------------------------------------------------------------------
+void FBasicToken::SetConstString(TCHAR* InString, int32 MaxLength)
+{
+	check(MaxLength>0);
+	ConstantType = CPT_String;
+	if( InString != String )
+	{
+		FCString::Strncpy(String, InString, MaxLength);
+	}
+	TokenType = TOKEN_Const;
+}
+
+//------------------------------------------------------------------------------
+FString FBasicToken::GetConstantValue() const
+{
+	if (TokenType == TOKEN_Const)
+	{
+		switch (ConstantType)
+		{
+		case CPT_Byte:
+			return FString::Printf(TEXT("%u"), Byte);
+		case CPT_Int:
+			return FString::Printf(TEXT("%i"), Int);
+		case CPT_Bool:
+			// Don't use GTrue/GFalse here because they can be localized
+			return FString::Printf(TEXT("%s"), NativeBool ? *(FName::GetEntry(NAME_TRUE)->GetPlainNameString()) : *(FName::GetEntry(NAME_FALSE)->GetPlainNameString()));
+		case CPT_Float:
+			return FString::Printf(TEXT("%f"), Float);
+		case CPT_Name:
+			return FString::Printf(TEXT("%s"), *(*(FName*)NameBytes).ToString());
+		case CPT_String:
+			return String;
+
+		// unsupported (parsing never produces a constant token of these types:
+		// CPT_Vector, ..., CPT_Int8, CPT_Int16, CPT_Int64, ..., CPT_Bool8, etc)
+		default:
+			return TEXT("InvalidTypeForAToken");
+		}
+	}
+	else
+	{
+		return TEXT("NotConstant");
+	}
+}
+
+//------------------------------------------------------------------------------
+bool FBasicToken::GetConstInt(int32& ValOut) const
+{
+	if(TokenType==TOKEN_Const && ConstantType==CPT_Int)
+	{
+		ValOut = Int;
+		return 1;
+	}
+	else if(TokenType==TOKEN_Const && ConstantType==CPT_Byte)
+	{
+		ValOut = Byte;
+		return 1;
+	}
+	else if(TokenType==TOKEN_Const && ConstantType==CPT_Float && Float==FMath::TruncToInt(Float))
+	{
+		ValOut = (int32)Float;
+		return 1;
+	}
+	else return 0;
+}
+
+/*******************************************************************************
+ * FBasicTokenParser::FErrorState
+*******************************************************************************/
+
+//------------------------------------------------------------------------------
+void FBasicTokenParser::FErrorState::Throw(bool bLogFatal) const
+{
+	if (State != NoError)
+	{
+		FString ErrorCodeStr;
+		switch (State)
+		{
+		case ParseError:
+			ErrorCodeStr = TEXT("ParseError");
+			break;
+		case RequireError:
+			ErrorCodeStr = TEXT("RequireError");
+			break;
+
+		default:
+			ErrorCodeStr.AppendInt(State.GetValue());
+		}
+
+		FString ErrorString = FString::Printf(TEXT("FBasicTokenParser Error (%s): %s"), *ErrorCodeStr, *Description.ToString());
+		// don't always log fatal (these could be presented as user facing errors), 
+		// but this is a good point to flip this bool on, to help catch the first 
+		// error in a possible chain of snowballing errors
+		if (bLogFatal)
+		{
+			UE_LOG(LogTokenParser, Fatal, TEXT("%s"), *ErrorString);
+		}
+		FError::Throwf(*ErrorString);
+	}
+}
+
+/*******************************************************************************
+ * FBasicTokenParser
+*******************************************************************************/
+
+//------------------------------------------------------------------------------
+void FBasicTokenParser::ResetParser(TCHAR const* SourceBuffer, int32 StartingLineNumber)
+{
+	Input     = SourceBuffer;
+	InputLen  = FCString::Strlen(Input);
+	InputPos  = 0;
+	PrevPos   = 0;
+	PrevLine  = 1;
+	InputLine = StartingLineNumber;
+
+	ClearCachedComment();
+	ClearErrorState();
+}
+
+//------------------------------------------------------------------------------
+void FBasicTokenParser::ClearCachedComment()
+{
+	// Can't call Reset as FString uses protected inheritance
+	PrevComment.Empty( PrevComment.Len() );
+}
+
+//------------------------------------------------------------------------------
+bool FBasicTokenParser::GetToken(FBasicToken& Token, bool bNoConsts/* = false*/)
+{
+	// if the parser is in a bad state, then don't continue parsing (who 
+	// knows what will happen!?)
+	if (!IsValid())
+	{
+		return false;
+	}
+
+	Token.TokenName	= NAME_None;
+	TCHAR c = GetLeadingChar();
+	TCHAR p = PeekChar();
+	if( c == 0 )
+	{
+		UngetChar();
+		return 0;
+	}
+	Token.StartPos		= PrevPos;
+	Token.StartLine		= PrevLine;
+	if( (c>='A' && c<='Z') || (c>='a' && c<='z') || (c=='_') )
+	{
+		// Alphanumeric token.
+		int32 Length=0;
+		do
+		{
+			Token.Identifier[Length++] = c;
+			if( Length >= NAME_SIZE )
+			{
+				Length = ((int32)NAME_SIZE) - 1;
+				Token.Identifier[Length]=0; // need this for the error description
+
+				FText ErrorDesc = FText::Format(LOCTEXT("IdTooLong", "Identifer ({0}...) exceeds maximum length of {1}"), FText::FromString(Token.Identifier), FText::AsNumber((int32)NAME_SIZE));
+				SetError(FErrorState::ParseError, ErrorDesc);
+
+				break;
+			}
+			c = GetChar();
+		} while( ((c>='A')&&(c<='Z')) || ((c>='a')&&(c<='z')) || ((c>='0')&&(c<='9')) || (c=='_') );
+		UngetChar();
+		Token.Identifier[Length]=0;
+
+		// Assume this is an identifier unless we find otherwise.
+		Token.TokenType = FBasicToken::TOKEN_Identifier;
+
+		// Lookup the token's global name.
+		Token.TokenName = FName( Token.Identifier, FNAME_Find, true );
+
+		// If const values are allowed, determine whether the identifier represents a constant
+		if ( !bNoConsts )
+		{
+			// See if the identifier is part of a vector, rotation or other struct constant.
+			// boolean true/false
+			if( Token.Matches(TEXT("true")) )
+			{
+				Token.SetConstBool(true);
+				return true;
+			}
+			else if( Token.Matches(TEXT("false")) )
+			{
+				Token.SetConstBool(false);
+				return true;
+			}
+		}
+
+		return IsValid();
+	}
+	// if const values are allowed, determine whether the non-identifier token represents a const
+	else if ( !bNoConsts && ((c>='0' && c<='9') || ((c=='+' || c=='-') && (p>='0' && p<='9'))) )
+	{
+		// Integer or floating point constant.
+		bool  bIsFloat = 0;
+		int32 Length   = 0;
+		bool  bIsHex   = 0;
+		do
+		{
+			if( c==TEXT('.') )
+			{
+				bIsFloat = true;
+			}
+			if( c==TEXT('X') || c == TEXT('x') )
+			{
+				bIsHex = true;
+			}
+
+			Token.Identifier[Length++] = c;
+			if( Length >= NAME_SIZE )
+			{
+				Length = ((int32)NAME_SIZE) - 1;
+				Token.Identifier[Length]=0; // need this for the error description
+
+				FText ErrorDesc = FText::Format(LOCTEXT("IdTooLong", "Identifer ({0}...) exceeds maximum length of {1}"), FText::FromString(Token.Identifier), FText::AsNumber((int32)NAME_SIZE));
+				SetError(FErrorState::ParseError, ErrorDesc);
+
+				break;
+			}
+			c = FChar::ToUpper(GetChar());
+		} while ((c >= TEXT('0') && c <= TEXT('9')) || (!bIsFloat && c == TEXT('.')) || (!bIsHex && c == TEXT('X')) || (bIsHex && c >= TEXT('A') && c <= TEXT('F')));
+
+		Token.Identifier[Length]=0;
+		if (!bIsFloat || c != 'F')
+		{
+			UngetChar();
+		}
+
+		if (bIsFloat)
+		{
+			Token.SetConstFloat( FCString::Atof(Token.Identifier) );
+		}
+		else if (bIsHex)
+		{
+			TCHAR* End = Token.Identifier + FCString::Strlen(Token.Identifier);
+			Token.SetConstInt( FCString::Strtoi(Token.Identifier,&End,0) );
+		}
+		else
+		{
+			Token.SetConstInt( FCString::Atoi(Token.Identifier) );
+		}
+		return IsValid();
+	}
+	else if( c=='"' )
+	{
+		// String constant.
+		TCHAR Temp[MAX_STRING_CONST_SIZE];
+		int32 Length=0;
+		c = GetChar(1);
+		while( (c!='"') && !IsEOL(c) )
+		{
+			if( c=='\\' )
+			{
+				c = GetChar(1);
+				if( IsEOL(c) )
+				{
+					break;
+				}
+				else if(c == 'n')
+				{
+					// Newline escape sequence.
+					c = '\n';
+				}
+			}
+			Temp[Length++] = c;
+			if( Length >= MAX_STRING_CONST_SIZE )
+			{
+				Length = ((int32)MAX_STRING_CONST_SIZE) - 1;
+				Temp[Length]=0; // need this for the error description
+
+				FText ErrorDesc = FText::Format(LOCTEXT("StringConstTooLong", "String constant ({0}...) exceeds maximum of {1} characters"), FText::FromString(Temp), FText::AsNumber((int32)MAX_STRING_CONST_SIZE));
+				SetError(FErrorState::ParseError, ErrorDesc);
+
+				c = TEXT('\"');				
+				break;
+			}
+			c = GetChar(1);
+		}
+		Temp[Length]=0;
+
+		if( c != '"' )
+		{
+			FText ErrorDesc = FText::Format(LOCTEXT("NoClosingQuote", "Unterminated quoted string ({0})"), FText::FromString(Temp));
+			SetError(FErrorState::ParseError, ErrorDesc);
+
+			UngetChar();
+		}
+
+		Token.SetConstString(Temp);
+		return IsValid();
+	}
+	else
+	{
+		// Symbol.
+		int32 Length=0;
+		Token.Identifier[Length++] = c;
+
+		// Handle special 2-character symbols.
+#define PAIR(cc,dd) ((c==cc)&&(d==dd)) /* Comparison macro for convenience */
+		TCHAR d = GetChar();
+		if
+		(	PAIR('<','<')
+		||	PAIR('>','>')
+		||	PAIR('!','=')
+		||	PAIR('<','=')
+		||	PAIR('>','=')
+		||	PAIR('+','+')
+		||	PAIR('-','-')
+		||	PAIR('+','=')
+		||	PAIR('-','=')
+		||	PAIR('*','=')
+		||	PAIR('/','=')
+		||	PAIR('&','&')
+		||	PAIR('|','|')
+		||	PAIR('^','^')
+		||	PAIR('=','=')
+		||	PAIR('*','*')
+		||	PAIR('~','=')
+		||	PAIR(':',':')
+		)
+		{
+			Token.Identifier[Length++] = d;
+			if( c=='>' && d=='>' )
+			{
+				if( GetChar()=='>' )
+					Token.Identifier[Length++] = '>';
+				else
+					UngetChar();
+			}
+		}
+		else UngetChar();
+#undef PAIR
+
+		Token.Identifier[Length] = 0;
+		Token.TokenType = FBasicToken::TOKEN_Symbol;
+
+		// Lookup the token's global name.
+		Token.TokenName = FName( Token.Identifier, FNAME_Find, true );
+
+		return true;
+	}
+}
+
+//------------------------------------------------------------------------------
+bool FBasicTokenParser::GetRawToken(FBasicToken& Token, TCHAR StopChar/* = TCHAR('\n')*/)
+{
+	// if the parser is in a bad state, then don't continue parsing (who 
+	// knows what will happen!?)
+	if (!IsValid())
+	{
+		return false;
+	}
+
+	// Get token after whitespace.
+	TCHAR Temp[MAX_STRING_CONST_SIZE];
+	int32 Length=0;
+	TCHAR c = GetLeadingChar();
+	while( !IsEOL(c) && c != StopChar )
+	{
+		if( (c=='/' && PeekChar()=='/') || (c=='/' && PeekChar()=='*') )
+		{
+			break;
+		}
+		Temp[Length++] = c;
+		if( Length >= MAX_STRING_CONST_SIZE )
+		{
+			Temp[Length] = 0;
+
+			FText ErrorDesc = FText::Format(LOCTEXT("IdTooLong", "Identifer ({0}...) exceeds maximum length of {1}"), FText::FromString(Temp), FText::AsNumber((int32)MAX_STRING_CONST_SIZE));
+			SetError(FErrorState::ParseError, ErrorDesc);
+		}
+		c = GetChar(true);
+	}
+	UngetChar();
+
+	// Get rid of trailing whitespace.
+	while( Length>0 && (Temp[Length-1]==' ' || Temp[Length-1]==9 ) )
+	{
+		Length--;
+	}
+	Temp[Length]=0;
+
+	Token.SetConstString(Temp);
+
+	return Length>0;
+}
+
+//------------------------------------------------------------------------------
+bool FBasicTokenParser::GetRawTokenRespectingQuotes(FBasicToken& Token, TCHAR StopChar/* = TCHAR('\n')*/)
+{
+	// if the parser is in a bad state, then don't continue parsing (who 
+	// knows what will happen!?)
+	if (!IsValid())
+	{
+		return false;
+	}
+
+	// Get token after whitespace.
+	TCHAR Temp[MAX_STRING_CONST_SIZE];
+	int32 Length=0;
+	TCHAR c = GetLeadingChar();
+
+	bool bInQuote = false;
+
+	while( !IsEOL(c) && ((c != StopChar) || bInQuote) )
+	{
+		if( (c=='/' && PeekChar()=='/') || (c=='/' && PeekChar()=='*') )
+		{
+			break;
+		}
+
+		if (c == '"')
+		{
+			bInQuote = !bInQuote;
+		}
+
+		Temp[Length++] = c;
+		if( Length >= MAX_STRING_CONST_SIZE )
+		{
+			Length = ((int32)MAX_STRING_CONST_SIZE) - 1;
+			Temp[Length]=0; // needs to happen for the error description below
+
+			FText ErrorDesc = FText::Format(LOCTEXT("IdTooLong", "Identifer ({0}...) exceeds maximum length of {1}"), FText::FromString(Temp), FText::AsNumber((int32)MAX_STRING_CONST_SIZE));
+			SetError(FErrorState::ParseError, ErrorDesc);
+
+			c = GetChar(true);
+			break;
+		}
+		c = GetChar(true);
+	}
+	UngetChar();
+
+	// Get rid of trailing whitespace.
+	while( Length>0 && (Temp[Length-1]==' ' || Temp[Length-1]==9 ) )
+	{
+		Length--;
+	}
+	Temp[Length]=0;
+
+	if (bInQuote)
+	{
+		FText ErrorDesc = FText::Format(LOCTEXT("NoClosingQuote", "Unterminated quoted string ({0})"), FText::FromString(Temp));
+		SetError(FErrorState::ParseError, ErrorDesc);
+	}
+	Token.SetConstString(Temp);
+
+	return Length>0 && IsValid();
+}
+
+//------------------------------------------------------------------------------
+bool FBasicTokenParser::GetIdentifier(FBasicToken& Token, bool bNoConsts)
+{
+	if (!GetToken(Token, bNoConsts))
+	{
+		return false;
+	}
+
+	if (Token.TokenType == FBasicToken::TOKEN_Identifier)
+	{
+		return true;
+	}
+
+	UngetToken(Token);
+	return false;
+}
+
+//------------------------------------------------------------------------------
+bool FBasicTokenParser::GetSymbol(FBasicToken& Token)
+{
+	if (!GetToken(Token))
+	{
+		return false;
+	}
+
+	if (Token.TokenType == FBasicToken::TOKEN_Symbol)
+	{
+		return true;
+	}
+
+	UngetToken(Token);
+	return false;
+}
+
+//------------------------------------------------------------------------------
+bool FBasicTokenParser::GetConstInt(int32& Result, TCHAR const* ErrorContext)
+{
+	FBasicToken Token;
+	if (GetToken(Token))
+	{
+		if (Token.GetConstInt(Result))
+		{
+			return true;
+		}
+		else
+		{
+			UngetToken(Token);
+		}
+	}
+
+	if (ErrorContext != NULL)
+	{
+		FText ErrorDesc = FText::Format(LOCTEXT("ContextualNoInt", "{0}: Missing expected integer constant"), FText::FromString(ErrorContext));
+		SetError(FErrorState::ParseError, ErrorDesc);
+	}
+	return false;
+}
+
+//------------------------------------------------------------------------------
+void FBasicTokenParser::UngetToken(FBasicToken& Token)
+{
+	InputPos = Token.StartPos;
+	InputLine = Token.StartLine;
+}
+
+//------------------------------------------------------------------------------
+TCHAR FBasicTokenParser::PeekChar()
+{
+	return (InputPos < InputLen) ? Input[InputPos] : 0;
+}
+
+//------------------------------------------------------------------------------
+TCHAR FBasicTokenParser::GetChar(bool bLiteral/* = false*/)
+{
+	// if the parser is in a bad state, then don't continue parsing (who 
+	// knows what will happen!?)... return a char signaling the end-of-stream
+	if (!IsValid())
+	{
+		return 0;
+	}
+
+	int32 CommentCount = 0;
+
+	PrevPos = InputPos;
+	PrevLine = InputLine;
+
+Loop:
+	const TCHAR c = Input[InputPos++];
+	if ( CommentCount > 0 )
+	{
+		// Record the character as a comment.
+		PrevComment += c;
+	}
+
+	if (c == TEXT('\n'))
+	{
+		InputLine++;
+	}
+	else if (!bLiteral)
+	{
+		const TCHAR NextChar = PeekChar();
+		if ( c==TEXT('/') && NextChar==TEXT('*') )
+		{
+			if ( CommentCount == 0 )
+			{
+				ClearCachedComment();
+				// Record the slash and star.
+				PrevComment += c;
+				PrevComment += NextChar;
+			}
+			CommentCount++;
+			InputPos++;
+			goto Loop;
+		}
+		else if( c==TEXT('*') && NextChar==TEXT('/') )
+		{
+			if (--CommentCount < 0)
+			{
+				ClearCachedComment();
+				SetError(FErrorState::ParseError, LOCTEXT("UnexpectedCommentClose", "Unexpected '*/' outside of comment"));
+			}
+			// Star already recorded; record the slash.
+			PrevComment += Input[InputPos];
+
+			InputPos++;
+			goto Loop;
+		}
+	}
+
+	if (CommentCount > 0)
+	{
+		if (c == 0)
+		{
+			ClearCachedComment();
+			SetError(FErrorState::ParseError, LOCTEXT("NoCommentClose", "No end to a comment by the end of the expression"));
+		}
+		goto Loop;
+	}
+	return c;
+}
+
+//------------------------------------------------------------------------------
+TCHAR FBasicTokenParser::GetLeadingChar()
+{
+	// if the parser is in a bad state, then don't continue parsing (who 
+	// knows what will happen!?)... return a char signaling the end-of-stream
+	if (!IsValid())
+	{
+		return 0;
+	}
+
+	TCHAR TrailingCommentNewline = 0;
+	for (;;)
+	{
+		bool MultipleNewlines = false;
+
+		TCHAR c;
+
+		// Skip blanks.
+		do
+		{
+			c = GetChar();
+
+			// Check if we've encountered another newline since the last one
+			if (c == TrailingCommentNewline)
+			{
+				MultipleNewlines = true;
+			}
+		} while (IsWhitespace(c));
+
+		if (c != TEXT('/') || PeekChar() != TEXT('/'))
+		{
+			return c;
+		}
+
+		// Clear the comment if we've encountered newlines since the last comment
+		if (MultipleNewlines)
+		{
+			ClearCachedComment();
+		}
+
+		// Record the first slash.  The first iteration of the loop will get the second slash.
+		PrevComment += c;
+
+		do
+		{
+			c = GetChar(true);
+			if (c == 0)
+				return c;
+			PrevComment += c;
+		} while (!IsEOL(c));
+
+		TrailingCommentNewline = c;
+
+		for (;;)
+		{
+			c = GetChar();
+			if (c == 0)
+				return c;
+			if (c == TrailingCommentNewline || !IsEOL(c))
+			{
+				UngetChar();
+				break;
+			}
+
+			PrevComment += c;
+		}
+	}
+}
+
+//------------------------------------------------------------------------------
+void FBasicTokenParser::UngetChar()
+{
+	InputPos = PrevPos;
+	InputLine = PrevLine;
+}
+
+//------------------------------------------------------------------------------
+bool FBasicTokenParser::IsEOL(TCHAR c)
+{
+	return c==TEXT('\n') || c==TEXT('\r') || c==0;
+}
+
+//------------------------------------------------------------------------------
+bool FBasicTokenParser::IsWhitespace(TCHAR c)
+{
+	return c==TEXT(' ') || c==TEXT('\t') || c==TEXT('\r') || c==TEXT('\n');
+}
+
+//------------------------------------------------------------------------------
+bool FBasicTokenParser::MatchIdentifier(FName Match)
+{
+	FBasicToken Token;
+	if (!GetToken(Token))
+	{
+		return false;
+	}
+
+	if ((Token.TokenType == FBasicToken::TOKEN_Identifier) && (Token.TokenName == Match))
+	{
+		return true;
+	}
+
+	UngetToken(Token);
+	return false;
+}
+
+//------------------------------------------------------------------------------
+bool FBasicTokenParser::MatchIdentifier(TCHAR const* Match)
+{
+	FBasicToken Token;
+	if (GetToken(Token))
+	{
+		if (Token.TokenType == FBasicToken::TOKEN_Identifier && FCString::Stricmp(Token.Identifier, Match) == 0)
+		{
+			return true;
+		}
+		else
+		{
+			UngetToken(Token);
+		}
+	}
+
+	return false;
+}
+
+//------------------------------------------------------------------------------
+bool FBasicTokenParser::PeekIdentifier(FName Match)
+{
+	FBasicToken Token;
+	if (!GetToken(Token, true))
+	{
+		return false;
+	}
+	UngetToken(Token);
+	return Token.TokenType == FBasicToken::TOKEN_Identifier && Token.TokenName == Match;
+}
+
+//------------------------------------------------------------------------------
+bool FBasicTokenParser::PeekIdentifier(const TCHAR* Match)
+{
+	FBasicToken Token;
+	if (!GetToken(Token, true))
+	{
+		return false;
+	}
+	UngetToken(Token);
+	return Token.TokenType == FBasicToken::TOKEN_Identifier && FCString::Stricmp(Token.Identifier, Match) == 0;
+}
+
+//------------------------------------------------------------------------------
+bool FBasicTokenParser::MatchSymbol(TCHAR const* Match)
+{
+	FBasicToken Token;
+
+	if (GetToken(Token, /*bNoConsts=*/ true))
+	{
+		if (Token.TokenType == FBasicToken::TOKEN_Symbol && !FCString::Stricmp(Token.Identifier, Match))
+		{
+			return true;
+		}
+		else
+		{
+			UngetToken(Token);
+		}
+	}
+
+	return false;
+}
+
+//------------------------------------------------------------------------------
+bool FBasicTokenParser::PeekSymbol(TCHAR const* Match)
+{
+	FBasicToken Token;
+	if (!GetToken(Token, true))
+	{
+		return false;
+	}
+	UngetToken(Token);
+
+	return Token.TokenType == FBasicToken::TOKEN_Symbol && FCString::Stricmp(Token.Identifier, Match) == 0;
+}
+
+//------------------------------------------------------------------------------
+bool FBasicTokenParser::RequireIdentifier(FName Match, TCHAR const* ErrorContext)
+{
+	if (!MatchIdentifier(Match))
+	{
+		FText ErrorDesc = FText::Format(LOCTEXT("MissingRequirement", "Missing '{0}' in {1}"), FText::FromName(Match), FText::FromString(ErrorContext));
+		SetError(FErrorState::RequireError, ErrorDesc);
+	}
+	return IsValid();
+}
+
+//------------------------------------------------------------------------------
+bool FBasicTokenParser::RequireIdentifier(TCHAR const* Match, TCHAR const* ErrorContext)
+{
+	if (!MatchIdentifier(Match))
+	{
+		FText ErrorDesc = FText::Format(LOCTEXT("MissingRequirement", "Missing '{0}' in {1}"), FText::FromString(Match), FText::FromString(ErrorContext));
+		SetError(FErrorState::RequireError, ErrorDesc);
+	}
+	return IsValid();
+}
+
+//------------------------------------------------------------------------------
+bool FBasicTokenParser::RequireSymbol(TCHAR const* Match, TCHAR const* ErrorContext)
+{
+	if (!MatchSymbol(Match))
+	{
+		FText ErrorDesc = FText::Format(LOCTEXT("MissingRequirement", "Missing '{0}' in {1}"), FText::FromString(Match), FText::FromString(ErrorContext));
+		SetError(FErrorState::RequireError, ErrorDesc);
+	}
+	return IsValid();
+}
+
+//------------------------------------------------------------------------------
+bool FBasicTokenParser::RequireSemi()
+{
+	if( !MatchSymbol(TEXT(";")) )
+	{
+		FText ErrorDesc = LOCTEXT("MissingSemiColon", "Missing ';'");
+
+		FBasicToken Token;
+		if( GetToken(Token) )
+		{
+			ErrorDesc = FText::Format(LOCTEXT("MissingSemiBefore", "Missing ';' before '{0}'"), FText::FromString(Token.Identifier));
+		}		
+		SetError(FErrorState::RequireError, ErrorDesc);
+	}
+	return IsValid();
+}
+
+//------------------------------------------------------------------------------
+void FBasicTokenParser::SetError(TEnumAsByte<FErrorState::EErrorType> ErrorCode, FText Description, bool bLogFatal)
+{
+	CurrentError.State = ErrorCode;
+	CurrentError.Description = Description;
+	CurrentError.Throw(bLogFatal);
+}
+
+//------------------------------------------------------------------------------
+FBasicTokenParser::FErrorState const& FBasicTokenParser::GetErrorState() const
+{
+	return CurrentError;
+}
+
+//------------------------------------------------------------------------------
+bool FBasicTokenParser::IsValid() const
+{
+	return (CurrentError.State == FErrorState::NoError);
+}
+
+//------------------------------------------------------------------------------
+void FBasicTokenParser::ClearErrorState()
+{
+	SetError(FErrorState::NoError, FText::GetEmpty());
+}
+
+#undef LOCTEXT_NAMESPACE
diff --git a/Engine/Plugins/Experimental/MathExpressions/Source/MathExpressions/Private/BasicTokenParser.h b/Engine/Plugins/Experimental/MathExpressions/Source/MathExpressions/Private/BasicTokenParser.h
new file mode 100644
--- /dev/null
+++ b/Engine/Plugins/Experimental/MathExpressions/Source/MathExpressions/Private/BasicTokenParser.h
@@ -0,0 +1,422 @@
+// Copyright 1998-2014 Epic Games, Inc. All Rights Reserved.
+
+#pragma once
+
+/*******************************************************************************
+ * FBasicToken
+*******************************************************************************/
+
+/** 
+ * Information regarding a token that was parsed from some expression string 
+ * (type, value, etc.).
+ */
+class FBasicToken //: public FPropertyBase
+{
+public:
+	FBasicToken(); // invokes InitToken()
+
+	/**
+	 * Resets this token, clearing any details that were previously set (allows
+	 * you to init the token to some predefined value type).
+	 * 
+	 * @param  InConstType	Optional param that lets you init this to some known value type. 
+	 */
+	void InitToken(EPropertyType InConstType = CPT_None);
+
+	/**
+	* Copies the properties from another token into this one.
+	*
+	* @param  Other		The token to copy properties from.
+	*/
+	void Clone(FBasicToken const& Other);
+
+	//--------------------------------------
+	// Queries
+	//--------------------------------------
+	
+	/** Determines if this token matches the specified name.*/
+	bool Matches(const FName& Name) const;
+
+	/** Determines if this token matches the specified string. */
+	bool Matches(const TCHAR* Str, ESearchCase::Type SearchCase = ESearchCase::IgnoreCase) const;
+
+	/** Determines if this token has the specified string as a prefix. */
+	bool StartsWith( TCHAR const* Str, bool bCaseSensitive = false) const;
+
+	/** Determines if this token has a boolean ConstantType.*/
+	bool IsBool() const;
+
+	//--------------------------------------
+	// Constant value setters
+	//--------------------------------------
+	
+	void SetConstInt(int32 InInt);
+	void SetConstBool(bool InBool);
+	void SetConstFloat(float InFloat);
+
+	void SetConstName(FName InName);
+	void SetConstString(TCHAR* InString, int32 MaxLength=MAX_STRING_CONST_SIZE);
+
+	//--------------------------------------
+	// Constant value getters
+	//--------------------------------------
+
+	/**
+	 * Retrieves an int value from this token (if it is a constant 
+	 * int/byte/float/etc. type).
+	 * 
+	 * @param  ValOut	Will hold the integer value if this is successfully.
+	 * @return True if this represents a constant that can be converted to an integer, otherwise false (ValOut is invalid).
+	 */
+	bool GetConstInt(int32& ValOut) const;
+
+	/**
+	 * If this represents a constant value, then this returns a string 
+	 * representing the value of said constant (formated according to the type).
+	 * 
+	 * @return A default error string if this is not a constant (or an 
+	 *         unsupported type), else a string containing the constant value.
+	 */
+	FString GetConstantValue() const;
+
+public:
+	enum ETokenType
+	{
+		TOKEN_None				= 0x00, // No token.
+		TOKEN_Identifier		= 0x01,	// Alphanumeric identifier.
+		TOKEN_Symbol			= 0x02,	// Symbol.
+		TOKEN_Const				= 0x03,	// A constant.
+		TOKEN_Max				= 0x0D
+	};
+	/** Type of this token. */
+	ETokenType TokenType;
+
+	/** Name of this token. */
+	FName TokenName;
+
+	/** Starting position in expression stream where this token came from. */
+	int32 StartPos;
+
+	/** Starting line in the expression. */
+	int32 StartLine;
+
+	/** Always valid. */
+	TCHAR Identifier[NAME_SIZE];
+
+	/** Only valid when TokenType is TOKEN_Const */
+	EPropertyType ConstantType;
+
+	/* TOKEN_Const values */
+	union
+	{
+		uint8 Byte;							// If CPT_Byte.
+		int32 Int;							// If CPT_Int.
+		bool  NativeBool;					// If CPT_Bool
+		float Float;						// If CPT_Float.
+		uint8 NameBytes[sizeof(FName)];		// If CPT_Name.
+		TCHAR String[MAX_STRING_CONST_SIZE];// If CPT_String
+	};
+};
+
+/*******************************************************************************
+ * FBasicTokenParser
+*******************************************************************************/
+
+/**
+ * Provides a base set of expression parsing functionality for sub-classes to 
+ * utilize when tokenizing a TCHAR stream (discerns operators from literals/ 
+ * variables, strips out comments, etc.).
+ */
+class FBasicTokenParser
+{
+protected:
+	/** 
+	 * FBasicTokenParser is conceptually abstract, and should only be instantiated 
+	 * through a subclass.
+	 */
+	FBasicTokenParser() {}
+
+protected:
+	/** 
+	 * Sets up this to start parsing a new expression string and resets any 
+	 * state lingering from the previous one.
+	 *
+	 * @param  SourceBuffer			The expression string you wish to parse next.
+	 * @param  StartingLineNumber	The base line number to start counting from as this parses.
+	 */
+	void ResetParser(const TCHAR* SourceBuffer, int32 StartingLineNumber = 1);
+
+	/** 
+	 * Clears out the most recent parsed comment. 
+	 */
+	void ClearCachedComment();
+
+	//--------------------------------------
+	// Tokenizing functions
+	//--------------------------------------
+
+	/**
+	 * Gets the next token from the input stream, advancing the variables which 
+	 * keep track of the current input position and line. 
+	 * 
+	 * @param  Token [out]	Receives the value of the parsed text; if Token is 
+	 *						pre-initialized, special logic is performed to attempt
+	 *						to evaluated Token in the context of that type. Useful
+	 *						for distinguishing between ambiguous symbols like enum tags.
+	 * @param  bNoConsts	Specify true to indicate that tokens representing 
+	 *						literal const values are not allowed.
+	 * @return True if a token was successfully processed, false otherwise.
+	 */
+	bool GetToken(FBasicToken& Token, bool bNoConsts = false);
+
+	/**
+	 * Put all text from the current position up to either EOL or the StopToken
+	 * into Token. Advances the compiler's current position.
+	 *
+	 * @param  Token [out]	Will contain the text that was parsed.
+	 * @param  StopChar		Stop processing when this character is reached.
+	 *
+	 * @return True if a token was successfully processed, false otherwise.
+	 */
+	bool GetRawToken(FBasicToken& Token, TCHAR StopChar = TCHAR('\n'));
+
+	/** 
+	 * Doesn't quit if StopChar is found inside a double-quoted string, but does 
+	 * NOT support quote escapes.
+	 *
+	 * @return True if a token was successfully processed, false otherwise.
+	 */
+	bool GetRawTokenRespectingQuotes(FBasicToken& Token, TCHAR StopChar = TCHAR('\n'));
+
+	/** 
+	 * Parse out an identifier from the expression stream.
+	 *
+	 * @return True if a token was successfully processed, false otherwise.
+	 */
+	bool GetIdentifier(FBasicToken& Token, bool bNoConsts = false);
+
+	/**
+	 * Parse out a symbol from the expression stream.
+	 *
+	 * @return True if a token was successfully processed, false otherwise.
+	 */
+	bool GetSymbol(FBasicToken& Token);
+
+	/**
+	 * Parse out an int constant from the expression stream.
+	 *
+	 * @return True if a int was successfully processed, false otherwise.
+	 */
+	bool GetConstInt(int32& Result, TCHAR const* ErrorContext = NULL);
+
+	/** 
+	 * Rolls back the line number and moves the parsing pointer back to where 
+	 * the specified token started.
+	 */
+	void UngetToken(FBasicToken& Token);
+
+	//--------------------------------------
+	// Low-level parsing functions
+	//--------------------------------------
+
+	/** 
+	 * Look at a single character from the input stream and return it, or 0=end.
+	 * Has no effect on the input stream.
+	 */
+	TCHAR PeekChar();
+
+	/** 
+	 * Get a single character from the input stream and return it, or 0=end. 
+	 */
+	TCHAR GetChar(bool bLiteral = false);
+
+	/** 
+	 * Skip past all spaces, tabs, and comments in the input stream. 
+	 */
+	TCHAR GetLeadingChar();
+
+	/** 
+	 * Unget the previous character retrieved with GetChar(). 
+	 */
+	void  UngetChar();
+
+	//--------------------------------------
+	// Match queries
+	//--------------------------------------
+
+	/**
+	 * Determines if the next token in the stream is an identifier that matches 
+	 * the specified FName  (and consumes it if it does). This is used primarily 
+	 * for checking for required symbols during parsing.
+	 * 
+	 * @param  Match	The name you wish to check for.
+	 * @return True if a match was consumed, false otherwise.
+	 */
+	bool MatchIdentifier(FName Match);
+
+	/**
+	 * Determines if the next token in the stream is an identifier that matches 
+	 * the specified string (and consumes it if it does).
+	 */
+	bool MatchIdentifier(TCHAR const* Match);
+
+	/**
+	 * Determines if the next token in the stream is an identifier that matches 
+	 * the specified name (does NOT consume it).
+	 */
+	bool PeekIdentifier(FName Match);
+
+	/**
+	 * Determines if the next token in the stream is an identifier that matches 
+	 * the specified string (does NOT consume it).
+	 */
+	bool PeekIdentifier(TCHAR const* Match);
+
+	/**
+	 * Determines if the next token in the stream is a symbol that matches 
+	 * the specified string (and consumes it if it does).
+	 */
+	bool MatchSymbol(TCHAR const* Match);
+
+	/**
+	 * Determines if the next token in the stream is a symbol that matches 
+	 * the specified string (does NOT consume it).
+	 */
+	bool PeekSymbol(TCHAR const* Match);
+
+	//--------------------------------------
+	// Requiring checks
+	//--------------------------------------
+
+	/**
+	 * Ensures that the next token in the stream is an identifier that matches 
+	 * the specified name (and errors out if it isn't).
+	 */
+	bool RequireIdentifier(FName Match, TCHAR const* ErrorContext);
+
+	/**
+	 * Ensures that the next token in the stream is an identifier that matches 
+	 * the specified string (and errors out if it isn't).
+	 */
+	bool RequireIdentifier(TCHAR const* Match, TCHAR const* ErrorContext);
+
+	/**
+	 * Ensures that the next token in the stream is a symbol that matches 
+	 * the specified string (and errors out if it isn't).
+	 */
+	bool RequireSymbol(TCHAR const* Match, TCHAR const* ErrorContext);
+
+	/**
+	 * Ensures that the next token in the stream is a semi-colon character (and 
+	 * errors out if it isn't).
+	 */
+	bool RequireSemi();
+
+	//--------------------------------------
+	// Character queries
+	//--------------------------------------
+
+	/**
+	 * Tests if a character is an end-of-line character.
+	 * 
+	 * @param  c	The character to test.
+	 * @return True if c is an end-of-line character, false otherwise.
+	 */
+	static bool IsEOL(TCHAR c);
+
+	/**
+	 * Tests if a character is a whitespace character.
+	 * 
+	 * @param  c	The character to test.
+	 * @return True if c is an whitespace character, false otherwise.
+	 */
+	static bool IsWhitespace(TCHAR c);
+
+	//--------------------------------------
+	// Error state 
+	//--------------------------------------
+
+public:
+	/**
+	 * Checks to see if an error has been caught by the internal error state.
+	 */
+	bool IsValid() const;
+
+	/** A struct for easily describing a specific error that occurred while parsing */
+	struct FErrorState
+	{
+		enum EErrorType
+		{
+			NoError,
+			ParseError,   // generic error that occurred while tokenizing the expression
+			RequireError, // explicit error that was invoked through a call to one of the Require methods
+
+			SubClassErrorStart, // the starting value for a subclass's 
+		};
+		
+		FErrorState() 
+			: State(NoError) 
+		{}
+
+		/** Will take the current error and throw it using FError::Throwf() */
+		void Throw(bool bLogFatal) const;
+
+		/** Should match up with EErrorType, but is extensible by subclasses 
+		    (comes from a subclass if >= SubClassErrorStart) */
+		TEnumAsByte<EErrorType> State; 
+		/** A detailed localized string, describing what exactly went wrong 
+		   (meant to be user facing) */
+		FText Description;
+	};
+
+	/**
+	 * Retrieves the parser's internal error state (so that users can interpret 
+	 * what might have gone wrong while tokenizing).
+	 */
+	FErrorState const& GetErrorState() const;
+
+protected:
+	/**
+	 * Takes the provided error and throws it (if the ErrorCode isn't 
+	 * FErrorState::NoError). 
+	 * 
+	 * @param  ErrorCode	Should match up with FErrorState::EErrorType, but is
+	 *						meant to be extensible by subclasses (so it could be 
+	 *						out of that range).
+	 * @param  Description	A detailed localized string, describing what exactly 
+	 *						went wrong (meant to be user facing).
+	 * @param  bLogFatal	Optional param where you can make a fatal log along 
+	 *						with the thrown error (in case error handling is turned off).
+	 */
+	void SetError(TEnumAsByte<FErrorState::EErrorType> ErrorCode, FText Description, bool bLogFatal = false);
+
+	/**
+	 * Resets the internal error state, such that the parser can continue on.
+	 */
+	void ClearErrorState();
+	
+protected:
+	/** Input text */
+	TCHAR const* Input;
+
+	/** Length of input text */
+	int32 InputLen;
+
+	/** Current position in text */
+	int32 InputPos;
+
+	/** Current line in text */
+	int32 InputLine;
+
+	/** Position previous to last GetChar() call */
+	int32 PrevPos;
+
+	/** Line previous to last GetChar() call. */
+	int32 PrevLine;
+
+	/** Previous comment parsed by GetChar() call. */
+	FString PrevComment;
+
+	/** Keeps track of the last error to occur */
+	FErrorState CurrentError;
+};
\ No newline at end of file
diff --git a/Engine/Plugins/Experimental/MathExpressions/Source/MathExpressions/Private/K2Node_MathExpression.cpp b/Engine/Plugins/Experimental/MathExpressions/Source/MathExpressions/Private/K2Node_MathExpression.cpp
--- a/Engine/Plugins/Experimental/MathExpressions/Source/MathExpressions/Private/K2Node_MathExpression.cpp
+++ b/Engine/Plugins/Experimental/MathExpressions/Source/MathExpressions/Private/K2Node_MathExpression.cpp
@@ -1,1023 +1,233 @@
 // Copyright 1998-2014 Epic Games, Inc. All Rights Reserved.
 
 #include "MathExpressionsPrivatePCH.h"
-#include "BlueprintEditorUtils.h"
+#include "K2Node_MathExpression.h"
+#include "BlueprintGraphClasses.h"
 #include "Kismet2NameValidators.h"
 #include "EdGraphUtilities.h"
 #include "K2ActionMenuBuilder.h" // for FK2ActionMenuBuilder::AddNewNodeAction()
+#include "BasicTokenParser.h"
+#include "UnrealMathUtility.h"
+#include "BlueprintEditorUtils.h"
 
-// DECLARE_LOG_CATEGORY_EXTERN(LogCompile, Log, All);
-// DEFINE_LOG_CATEGORY(LogCompile);
-
-//@TODO: Hack
-bool GMessingWithMathNode = false;
-
-//
-// Token types.
-//
-enum ETokenType
-{
-	// No token.
-	TOKEN_None,
-
-	// Alphanumeric identifier.
-	TOKEN_Identifier,
-
-	// Symbol.
-	TOKEN_Symbol,
-
-	// A constant
-	TOKEN_Const
-};
-
-enum EConstantType
-{
-	CONST_Integer,
-	CONST_Boolean,
-	CONST_Float,
-	CONST_String
-};
-
-/////////////////////////////////////////////////////
-// FToken
-
-//
-// Information about a token that was just parsed.
-//
-class FToken
-{
-public:
-	/** Type of token. */
-	ETokenType TokenType;
-
-	/** Name of token. */
-	FName TokenName;
-
-	/** Starting position in script where this token came from. */
-	int32 StartPos;
-
-	/** Starting line in script. */
-	int32 StartLine;
-
-	/** Always valid. */
-	TCHAR Identifier[NAME_SIZE];
-
-	/** Type of constant (only valid if TokenType == TOKEN_Const) */
-	EConstantType ConstantType;
-
-	// Value of constant (only valid if TokenType == TOKEN_Const) */	
-	union
-	{
-		// TOKEN_Const values.
-		int32 Int;								 // If CONST_Integer
-		bool NativeBool;						 // if CONST_Boolean
-		float Float;							 // If CONST_Float
-		TCHAR String[MAX_STRING_CONST_SIZE];	 // If CONST_String
-	};
-
-	// Constructor
-	FToken()
-	{
-		TokenType = TOKEN_None;
-		ConstantType = CONST_Integer;
-		TokenName = NAME_None;
-		StartPos = 0;
-		StartLine = 0;
-		*Identifier = 0;
-		FMemory::Memzero(String, sizeof(Identifier));
-	}
-
-	FString GetConstantValue() const
-	{
-		if (TokenType == TOKEN_Const)
-		{
-			switch (ConstantType)
-			{
-			case CONST_Integer:
-				return FString::Printf(TEXT("%i"), Int);
-			case CONST_Boolean:	
-				// Don't use GTrue/GFalse here because they can be localized
-				return FString::Printf(TEXT("%s"), NativeBool ? *(FName::GetEntry(NAME_TRUE)->GetPlainNameString()) : *(FName::GetEntry(NAME_FALSE)->GetPlainNameString()));
-			case CONST_Float:
-				return FString::Printf(TEXT("%f"), Float);
-			case CONST_String:
-				return FString::Printf(TEXT("\"%s\""), String);
-				
-			default:
-				return TEXT("???");
-			}
-		}
-		else
-		{
-			return TEXT("NotAConstant");
-		}
-	}
-
-	bool Matches(const TCHAR* Str, bool bCaseSensitive = false) const
-	{
-		return ((TokenType == TOKEN_Identifier) || (TokenType == TOKEN_Symbol)) && (bCaseSensitive ? (!FCString::Strcmp(Identifier, Str)) : (!FCString::Stricmp(Identifier, Str)));
-	}
-
-	bool Matches(const FName& Name) const
-	{
-		return (TokenType == TOKEN_Identifier) && (TokenName == Name);
-	}
-
-	bool StartsWith(const TCHAR* Str, bool bCaseSensitive = false) const
-	{
-		const int32 StrLength = FCString::Strlen(Str);
-		return ((TokenType == TOKEN_Identifier) || (TokenType == TOKEN_Symbol)) && (bCaseSensitive ? (!FCString::Strncmp(Identifier, Str, StrLength)) : (!FCString::Strnicmp(Identifier, Str, StrLength)));
-	}
-
-	void SetConstInt(int32 InInt)
-	{
-		ConstantType = CONST_Integer;
-		Int = InInt;
-		TokenType = TOKEN_Const;
-	}
-
-	void SetConstBool(bool InBool)
-	{
-		ConstantType = CONST_Boolean;
-		NativeBool = InBool;
-		TokenType = TOKEN_Const;
-	}
-
-	void SetConstFloat(float InFloat)
-	{
-		ConstantType = CONST_Float;
-		Float = InFloat;
-		TokenType = TOKEN_Const;
-	}
-
-	void SetConstString(TCHAR* InString, int32 MaxLength = MAX_STRING_CONST_SIZE)
-	{
-		check(MaxLength>0);
-		ConstantType = CONST_String;
-		if (InString != String)
-		{
-			FCString::Strncpy( String, InString, MaxLength );
-		}
-		TokenType = TOKEN_Const;
-	}
-};
-
-
-
-
-
-/////////////////////////////////////////////////////
-// FBaseParser
-
-//
-// Base class of header parsers.
-//
-
-class FBaseParser
-{
-protected:
-	FBaseParser();
-protected:
-	// Input text.
-	const TCHAR* Input;
-
-	// Length of input text.
-	int32 InputLen;
-
-	// Current position in text.
-	int32 InputPos;
-
-	// Current line in text.
-	int32 InputLine;
-
-	// Position previous to last GetChar() call.
-	int32 PrevPos;
-
-	// Line previous to last GetChar() call.
-	int32 PrevLine;
-
-	// Previous comment parsed by GetChar() call.
-	FString PrevComment;
-
-	// true once PrevComment has been formatted.
-	bool bPrevCommentFormatted;
-
-	// Number of statements parsed.
-	int32 StatementsParsed;
-
-	// Total number of lines parsed.
-	int32 LinesParsed;
-protected:
-	void ResetParser(const TCHAR* SourceBuffer, int32 StartingLineNumber = 1);
-
-	// Low-level parsing functions.
-	TCHAR GetChar( bool Literal = false );
-	TCHAR PeekChar();
-	TCHAR GetLeadingChar();
-	void UngetChar();
-	bool IsEOL( TCHAR c );
-
-	/**
-	 * Gets the next token from the input stream, advancing the variables which keep track of the current input position and line.
-	 *
-	 * @param	Token			receives the value of the parsed text; if Token is pre-initialized, special logic is performed
-	 *							to attempt to evaluated Token in the context of that type.  Useful for distinguishing between ambigous symbols
-	 *							like enum tags.
-	 * @param	NoConsts		specify true to indicate that tokens representing literal const values are not allowed.
-	 *
-	 * @return	true if a token was successfully processed, false otherwise.
-	 */
-	bool GetToken(FToken& Token, bool bNoConsts = false);
-
-	/**
-	 * Put all text from the current position up to either EOL or the StopToken
-	 * into Token.  Advances the compiler's current position.
-	 *
-	 * @param	Token	[out] will contain the text that was parsed
-	 * @param	StopChar	stop processing when this character is reached
-	 *
-	 * @return	true if a token was parsed
-	 */
-	bool GetRawToken(FToken& Token, TCHAR StopChar = TCHAR('\n'));
-
-	// Doesn't quit if StopChar is found inside a double-quoted string, but does not support quote escapes
-	bool GetRawTokenRespectingQuotes(FToken& Token, TCHAR StopChar = TCHAR('\n'));
-
-	void UngetToken(FToken& Token);
-	bool GetIdentifier(FToken& Token, bool bNoConsts = false);
-	bool GetSymbol(FToken& Token);
-
-	// Matching predefined text.
-	bool MatchIdentifier(FName Match);
-	bool MatchIdentifier(const TCHAR* Match);
-	bool PeekIdentifier(FName Match);
-	bool PeekIdentifier(const TCHAR* Match);
-	bool MatchSymbol( const TCHAR* Match);
-	void MatchSemi();
-	bool PeekSymbol(const TCHAR* Match);
-
-	// Requiring predefined text.
-	void RequireIdentifier(FName Match, const TCHAR* Tag);
-	void RequireIdentifier(const TCHAR* Match, const TCHAR* Tag);
-	void RequireSymbol(const TCHAR* Match, const TCHAR* Tag);
-
-	/** Clears out the stored comment. */
-	void ClearComment();
-};
-
-
-
-
-
-
-
-
-/////////////////////////////////////////////////////
-// FBaseParser
-
-FBaseParser::FBaseParser()
-	: StatementsParsed(0)
-	, LinesParsed(0)
-{
-}
-
-void FBaseParser::ResetParser(const TCHAR* SourceBuffer, int32 StartingLineNumber)
-{
-	Input = SourceBuffer;
-	InputLen = FCString::Strlen(Input);
-	InputPos = 0;
-	PrevPos = 0;
-	PrevLine = 1;
-	InputLine = StartingLineNumber;
-}
-
-/*-----------------------------------------------------------------------------
-	Single-character processing.
------------------------------------------------------------------------------*/
-
-//
-// Get a single character from the input stream and return it, or 0=end.
-//
-TCHAR FBaseParser::GetChar(bool bLiteral)
-{
-	int32 CommentCount = 0;
-
-	PrevPos = InputPos;
-	PrevLine = InputLine;
-
-Loop:
-	const TCHAR c = Input[InputPos++];
-	if ( CommentCount > 0 )
-	{
-		// Record the character as a comment.
-		PrevComment += c;
-	}
-
-	if (c == TEXT('\n'))
-	{
-		InputLine++;
-	}
-	else if (!bLiteral)
-	{
-		const TCHAR NextChar = PeekChar();
-		if ( c==TEXT('/') && NextChar==TEXT('*') )
-		{
-			if ( CommentCount == 0 )
-			{
-				ClearComment();
-				// Record the slash and star.
-				PrevComment += c;
-				PrevComment += NextChar;
-			}
-			CommentCount++;
-			InputPos++;
-			goto Loop;
-		}
-		else if( c==TEXT('*') && NextChar==TEXT('/') )
-		{
-			if (--CommentCount < 0)
-			{
-				ClearComment();
-				FError::Throwf(TEXT("Unexpected '*/' outside of comment") );
-			}
-			// Star already recorded; record the slash.
-			PrevComment += Input[InputPos];
-
-			InputPos++;
-			goto Loop;
-		}
-	}
-
-	if (CommentCount > 0)
-	{
-		if (c == 0)
-		{
-			ClearComment();
-			FError::Throwf(TEXT("End of class header encountered inside comment") );
-		}
-		goto Loop;
-	}
-	return c;
-}
-
-//
-// Unget the previous character retrieved with GetChar().
-//
-void FBaseParser::UngetChar()
-{
-	InputPos = PrevPos;
-	InputLine = PrevLine;
-}
-
-//
-// Look at a single character from the input stream and return it, or 0=end.
-// Has no effect on the input stream.
-//
-TCHAR FBaseParser::PeekChar()
-{
-	return (InputPos < InputLen) ? Input[InputPos] : 0;
-}
-
-//
-// Skip past all spaces and tabs in the input stream.
-//
-TCHAR FBaseParser::GetLeadingChar()
-{
-	// Skip blanks.
-	TCHAR c;
-	Skip1:
-	do
-	{
-		c = GetChar();
-	} while( c==TEXT(' ') || c==TEXT('\t') || c==TEXT('\r') || c==TEXT('\n') );
-
-	if ((c == TEXT('/')) && (PeekChar() == TEXT('/')))
-	{
-		// Record the first slash.  The first iteration of the loop will get the second slash.
-		PrevComment += c;
-		do
-		{
-			c = GetChar(true);
-			if (c == 0)
-			{
-				return c;
-			}
-			PrevComment += c;
-		} while((c != TEXT('\r')) && (c != TEXT('\n')));
-
-		goto Skip1;
-	}
-	return c;
-}
-
-//
-// Return 1 if input as a valid end-of-line character, or 0 if not.
-// EOL characters are: Comment, CR, linefeed, 0 (end-of-file mark)
-//
-bool FBaseParser::IsEOL( TCHAR c )
-{
-	return c==TEXT('\n') || c==TEXT('\r') || c==0;
-}
+#define LOCTEXT_NAMESPACE "K2Node"
 
-/*-----------------------------------------------------------------------------
-	Tokenizing.
------------------------------------------------------------------------------*/
+/*******************************************************************************
+ * Static Helpers
+*******************************************************************************/
 
-// Gets the next token from the input stream, advancing the variables which keep track of the current input position and line.
-bool FBaseParser::GetToken( FToken& Token, bool bNoConsts/*=false*/ )
+/**
+ * Helper function for deleting all the nodes from a specified graph. Does not
+ * delete any tunnel in/out nodes (to preserve the tunnel).
+ * 
+ * @param  Graph	The graph you want to delete nodes in.
+ */
+static void DeleteGeneratedNodesInGraph(UEdGraph* Graph)
 {
-	Token.TokenName	= NAME_None;
-	TCHAR c = GetLeadingChar();
-	TCHAR p = PeekChar();
-	if( c == 0 )
-	{
-		UngetChar();
-		return 0;
-	}
-	Token.StartPos		= PrevPos;
-	Token.StartLine		= PrevLine;
-	if( (c>='A' && c<='Z') || (c>='a' && c<='z') || (c=='_') )
-	{
-		// Alphanumeric token.
-		int32 Length=0;
-		do
-		{
-			Token.Identifier[Length++] = c;
-			if( Length >= NAME_SIZE )
-			{
-				FError::Throwf(TEXT("Identifer length exceeds maximum of %i"), (int32)NAME_SIZE);
-				Length = ((int32)NAME_SIZE) - 1;
-				break;
-			}
-			c = GetChar();
-		} while( ((c>='A')&&(c<='Z')) || ((c>='a')&&(c<='z')) || ((c>='0')&&(c<='9')) || (c=='_') );
-		UngetChar();
-		Token.Identifier[Length]=0;
-
-		// Assume this is an identifier unless we find otherwise.
-		Token.TokenType = TOKEN_Identifier;
-
-		// Lookup the token's global name.
-		Token.TokenName = FName( Token.Identifier, FNAME_Find, true );
-
-		// If const values are allowed, determine whether the identifier represents a constant
-		if ( !bNoConsts )
-		{
-			// See if the identifier is part of a vector, rotation or other struct constant.
-			// boolean true/false
-			if( Token.Matches(TEXT("true")) )
-			{
-				Token.SetConstBool(true);
-				return true;
-			}
-			else if( Token.Matches(TEXT("false")) )
-			{
-				Token.SetConstBool(false);
-				return true;
-			}
-		}
-		return true;
-	}
-
-	// if const values are allowed, determine whether the non-identifier token represents a const
-	else if ( !bNoConsts && ((c>='0' && c<='9') || ((c=='+' || c=='-') && (p>='0' && p<='9'))) )
+	UBlueprint* Blueprint = FBlueprintEditorUtils::FindBlueprintForGraphChecked(Graph);
+	for (int32 NodeIndex = 0; NodeIndex < Graph->Nodes.Num(); )
 	{
-		// Integer or floating point constant.
-		bool bIsFloat = 0;
-		int32 Length = 0;
-		bool bIsHex = 0;
-		do
-		{
-			if( c==TEXT('.') )
-			{
-				bIsFloat = true;
-			}
-			if( c==TEXT('X') || c == TEXT('x') )
-			{
-				bIsHex = true;
-			}
-
-			Token.Identifier[Length++] = c;
-			if( Length >= NAME_SIZE )
-			{
-				FError::Throwf(TEXT("Number length exceeds maximum of %i "), (int32)NAME_SIZE );
-				Length = ((int32)NAME_SIZE) - 1;
-				break;
-			}
-			c = FChar::ToUpper(GetChar());
-		} while ((c >= TEXT('0') && c <= TEXT('9')) || (!bIsFloat && c == TEXT('.')) || (!bIsHex && c == TEXT('X')) || (bIsHex && c >= TEXT('A') && c <= TEXT('F')));
-
-		Token.Identifier[Length]=0;
-		if (!bIsFloat || c != 'F')
-		{
-			UngetChar();
-		}
-
-		if (bIsFloat)
-		{
-			Token.SetConstFloat( FCString::Atof(Token.Identifier) );
-		}
-		else if (bIsHex)
+		UEdGraphNode* Node = Graph->Nodes[NodeIndex];
+		if (ExactCast<UK2Node_Tunnel>(Node) != NULL)
 		{
-			TCHAR* End = Token.Identifier + FCString::Strlen(Token.Identifier);
-			Token.SetConstInt( FCString::Strtoi(Token.Identifier,&End,0) );
+			++NodeIndex;
 		}
 		else
 		{
-			Token.SetConstInt( FCString::Atoi(Token.Identifier) );
-		}
-		return true;
-	}
-//@TODO: Support single character literal parsing - the code below is for FNames, which is wrong
-// 	else if( !bNoConsts && c=='\'' )
-// 	{
-// 		// Name constant.
-// 		int32 Length=0;
-// 		c = GetChar();
-// 		while( (c>='A' && c<='Z') || (c>='a' && c<='z') || (c>='0' && c<='9') || (c=='_') || (c=='-') || (c==' ') ) //@FIXME: space in names should be illegal!
-// 		{
-// 			Token.Identifier[Length++] = c;
-// 			if( Length >= NAME_SIZE )
-// 			{
-// 				FError::Throwf(TEXT("Name length exceeds maximum of %i"), (int32)NAME_SIZE );
-// 				// trick the error a few lines down
-// 				c = TEXT('\'');
-// 				Length = ((int32)NAME_SIZE) - 1;
-// 				break;
-// 			}
-// 			c = GetChar();
-// 		}
-// 		if( c != '\'' )
-// 		{
-// 			UngetChar();
-// 			FError::Throwf(TEXT("Illegal character in name") );
-// 		}
-// 		Token.Identifier[Length]=0;
-// 
-// 		// Make constant name.
-// 		Token.SetConstName( FName(Token.Identifier) );
-// 		return true;
-// 	}
-	else if( c=='"' )
-	{
-		// String constant.
-		TCHAR Temp[MAX_STRING_CONST_SIZE];
-		int32 Length=0;
-		c = GetChar(1);
-		while( (c!='"') && !IsEOL(c) )
-		{
-			if( c=='\\' )
-			{
-				c = GetChar(1);
-				if( IsEOL(c) )
-				{
-					break;
-				}
-				else if(c == 'n')
-				{
-					// Newline escape sequence.
-					c = '\n';
-				}
-			}
-			Temp[Length++] = c;
-			if( Length >= MAX_STRING_CONST_SIZE )
-			{
-				FError::Throwf(TEXT("String constant exceeds maximum of %i characters"), (int32)MAX_STRING_CONST_SIZE );
-				c = TEXT('\"');
-				Length = ((int32)MAX_STRING_CONST_SIZE) - 1;
-				break;
-			}
-			c = GetChar(1);
-		}
-		Temp[Length]=0;
-
-		if( c != '"' )
-		{
-			FError::Throwf(TEXT("Unterminated string constant: %s"), Temp);
-			UngetChar();
-		}
-
-		Token.SetConstString(Temp);
-		return true;
-	}
-	else
-	{
-		// Symbol.
-		int32 Length=0;
-		Token.Identifier[Length++] = c;
-
-		// Handle special 2-character symbols.
-		#define PAIR(cc,dd) ((c==cc)&&(d==dd)) /* Comparison macro for convenience */
-		TCHAR d = GetChar();
-		if
-		(	PAIR('<','<')
-		||	PAIR('>','>')
-		||	PAIR('!','=')
-		||	PAIR('<','=')
-		||	PAIR('>','=')
-		||	PAIR('+','+')
-		||	PAIR('-','-')
-		||	PAIR('+','=')
-		||	PAIR('-','=')
-		||	PAIR('*','=')
-		||	PAIR('/','=')
-		||	PAIR('&','&')
-		||	PAIR('|','|')
-		||	PAIR('^','^')
-		||	PAIR('=','=')
-		||	PAIR('*','*')
-		||	PAIR('~','=')
-		||	PAIR(':',':')
-		)
-		{
-			Token.Identifier[Length++] = d;
-			if( c=='>' && d=='>' )
-			{
-				if( GetChar()=='>' )
-					Token.Identifier[Length++] = '>';
-				else
-					UngetChar();
-			}
-		}
-		else UngetChar();
-		#undef PAIR
-
-		Token.Identifier[Length] = 0;
-		Token.TokenType = TOKEN_Symbol;
-
-		// Lookup the token's global name.
-		Token.TokenName = FName( Token.Identifier, FNAME_Find, true );
-
-		return true;
-	}
-}
-
-bool FBaseParser::GetRawTokenRespectingQuotes( FToken& Token, TCHAR StopChar /* = TCHAR('\n') */ )
-{
-	// Get token after whitespace.
-	TCHAR Temp[MAX_STRING_CONST_SIZE];
-	int32 Length=0;
-	TCHAR c = GetLeadingChar();
-
-	bool bInQuote = false;
-
-	while( !IsEOL(c) && ((c != StopChar) || bInQuote) )
-	{
-		if( (c=='/' && PeekChar()=='/') || (c=='/' && PeekChar()=='*') )
-		{
-			break;
-		}
-
-		if (c == '"')
-		{
-			bInQuote = !bInQuote;
-		}
-
-		Temp[Length++] = c;
-		if( Length >= MAX_STRING_CONST_SIZE )
-		{
-			FError::Throwf(TEXT("Identifier exceeds maximum of %i characters"), (int32)MAX_STRING_CONST_SIZE );
-			c = GetChar(true);
-			Length = ((int32)MAX_STRING_CONST_SIZE) - 1;
-			break;
-		}
-		c = GetChar(true);
-	}
-	UngetChar();
-
-	if (bInQuote)
-	{
-		FError::Throwf(TEXT("Unterminated quoted string"));
-	}
-
-	// Get rid of trailing whitespace.
-	while( Length>0 && (Temp[Length-1]==' ' || Temp[Length-1]==9 ) )
-	{
-		Length--;
-	}
-	Temp[Length]=0;
-
-	Token.SetConstString(Temp);
-
-	return Length>0;
-}
-
-bool FBaseParser::GetRawToken( FToken& Token, TCHAR StopChar /* = TCHAR('\n') */ )
-{
-	// Get token after whitespace.
-	TCHAR Temp[MAX_STRING_CONST_SIZE];
-	int32 Length=0;
-	TCHAR c = GetLeadingChar();
-	while( !IsEOL(c) && c != StopChar )
-	{
-		if( (c=='/' && PeekChar()=='/') || (c=='/' && PeekChar()=='*') )
-		{
-			break;
-		}
-		Temp[Length++] = c;
-		if( Length >= MAX_STRING_CONST_SIZE )
-		{
-			FError::Throwf(TEXT("Identifier exceeds maximum of %i characters"), (int32)MAX_STRING_CONST_SIZE );
+			FBlueprintEditorUtils::RemoveNode(Blueprint, Node, true);
 		}
-		c = GetChar(true);
-	}
-	UngetChar();
-
-	// Get rid of trailing whitespace.
-	while( Length>0 && (Temp[Length-1]==' ' || Temp[Length-1]==9 ) )
-	{
-		Length--;
 	}
-	Temp[Length]=0;
-
-	Token.SetConstString(Temp);
-
-	return Length>0;
 }
 
-//
-// Get an identifier token, return 1 if gotten, 0 if not.
-//
-bool FBaseParser::GetIdentifier( FToken& Token, bool bNoConsts )
+/**
+ * If the specified type is a "byte" type, then this will modify it to 
+ * an "int". Helps when trying to match function signatures.
+ * 
+ * @param  InOutType	The type you want to attempt to promote.
+ * @return True if the type was modified, false if not.
+ */
+static bool PromoteByteToInt(FEdGraphPinType& InOutType)
 {
-	if (!GetToken(Token, bNoConsts))
-	{
-		return false;
-	}
-
-	if (Token.TokenType == TOKEN_Identifier)
+	const UEdGraphSchema_K2* Schema = GetDefault<UEdGraphSchema_K2>();
+	if (InOutType.PinCategory == Schema->PC_Byte)
 	{
+		InOutType.PinCategory = Schema->PC_Int;
+		InOutType.PinSubCategoryObject = NULL;
 		return true;
 	}
-
-	UngetToken(Token);
 	return false;
 }
 
-//
-// Get a symbol token, return 1 if gotten, 0 if not.
-//
-bool FBaseParser::GetSymbol( FToken& Token )
-{
-	if (!GetToken(Token))
-	{
-		return false;
-	}
-
-	if( Token.TokenType == TOKEN_Symbol )
-	{
-		return true;
-	}
-
-	UngetToken(Token);
-	return false;
-}
-
-bool FBaseParser::MatchSymbol( const TCHAR* Match )
-{
-	FToken Token;
-
-	if (GetToken(Token, /*bNoConsts=*/ true))
-	{
-		if (Token.TokenType==TOKEN_Symbol && !FCString::Stricmp(Token.Identifier, Match))
-		{
-			return true;
-		}
-		else
-		{
-			UngetToken(Token);
-		}
-	}
-
-	return false;
-}
-
-//
-// Get a specific identifier and return 1 if gotten, 0 if not.
-// This is used primarily for checking for required symbols during compilation.
-//
-bool FBaseParser::MatchIdentifier( FName Match )
+/**
+* If the specified type is a "int" type, then this will modify it to
+* a "float". Helps when trying to match function signatures.
+*
+* @param  InOutType	The type you want to attempt to promote.
+* @return True if the type was modified, false if not.
+*/
+static bool PromoteIntToFloat(FEdGraphPinType& InOutType)
 {
-	FToken Token;
-	if (!GetToken(Token))
-	{
-		return false;
-	}
-
-	if ((Token.TokenType == TOKEN_Identifier) && (Token.TokenName == Match))
+	const UEdGraphSchema_K2* Schema = GetDefault<UEdGraphSchema_K2>();
+	if (InOutType.PinCategory == Schema->PC_Int)
 	{
+		InOutType.PinCategory = Schema->PC_Float;
+		InOutType.PinSubCategoryObject = NULL;
 		return true;
 	}
+	return false;
+}
 
-	UngetToken(Token);
-	return false;
-}
-
-bool FBaseParser::MatchIdentifier( const TCHAR* Match )
-{
-	FToken Token;
-	if (GetToken(Token))
-	{
-		if( Token.TokenType==TOKEN_Identifier && FCString::Stricmp(Token.Identifier,Match)==0 )
-		{
-			return true;
-		}
-		else
-		{
-			UngetToken(Token);
-		}
-	}
-	
-	return false;
-}
-
-
-void FBaseParser::MatchSemi()
-{
-	if( !MatchSymbol(TEXT(";")) )
-	{
-		FToken Token;
-		if( GetToken(Token) )
-		{
-			FError::Throwf(TEXT("Missing ';' before '%s'"), Token.Identifier );
-		}
-		else
-		{
-			FError::Throwf(TEXT("Missing ';'") );
-		}
-	}
-}
-
-
-//
-// Peek ahead and see if a symbol follows in the stream.
-//
-bool FBaseParser::PeekSymbol( const TCHAR* Match )
-{
-	FToken Token;
-	if (!GetToken(Token, true))
-	{
-		return false;
-	}
-	UngetToken(Token);
-
-	return Token.TokenType==TOKEN_Symbol && FCString::Stricmp(Token.Identifier,Match)==0;
-}
-
-//
-// Peek ahead and see if an identifier follows in the stream.
-//
-bool FBaseParser::PeekIdentifier( FName Match )
-{
-	FToken Token;
-	if (!GetToken(Token, true))
-	{
-		return false;
-	}
-	UngetToken(Token);
-	return Token.TokenType==TOKEN_Identifier && Token.TokenName==Match;
-}
-
-bool FBaseParser::PeekIdentifier( const TCHAR* Match )
-{
-	FToken Token;
-	if (!GetToken(Token, true))
-	{
-		return false;
-	}
-	UngetToken(Token);
-	return Token.TokenType==TOKEN_Identifier && FCString::Stricmp(Token.Identifier,Match)==0;
-}
-
-//
-// Unget the most recently gotten token.
-//
-void FBaseParser::UngetToken( FToken& Token )
-{
-	InputPos = Token.StartPos;
-	InputLine = Token.StartLine;
-}
-
-//
-// Require a symbol.
-//
-void FBaseParser::RequireSymbol( const TCHAR* Match, const TCHAR* Tag )
-{
-	if (!MatchSymbol(Match))
-	{
-		FError::Throwf(TEXT("Missing '%s' in %s"), Match, Tag );
-	}
-}
-
-//
-// Require an identifier.
-//
-void FBaseParser::RequireIdentifier( FName Match, const TCHAR* Tag )
-{
-	if (!MatchIdentifier(Match))
-	{
-		FError::Throwf(TEXT("Missing '%s' in %s"), *Match.ToString(), Tag );
-	}
-}
-
-void FBaseParser::RequireIdentifier( const TCHAR* Match, const TCHAR* Tag )
-{
-	if (!MatchIdentifier(Match))
-	{
-		FError::Throwf(TEXT("Missing '%s' in %s"), Match, Tag );
-	}
-}
-
-// Clears out the stored comment.
-void FBaseParser::ClearComment()
-{
-	// Can't call Reset as FString uses protected inheritance
-	PrevComment.Empty( PrevComment.Len() );
-	bPrevCommentFormatted = false;
-}
-
-
-
-
-
-
-
-
-
-
-
-
-
-
-
-
-
-
-
-
-
-class FTreeContext
+/**
+ * Sets or clears the error on a specific node. If the ErrorText is empty, then
+ * it resets the error state on the node. If actual error text is supplied, 
+ * then the node is flagged as having an error, and the string is appended to 
+ * the node's error message.
+ * 
+ * @param  Node			The node you wish to modify.
+ * @param  ErrorText	The text you want to append to the node's error message (empty if you want to clear any errors).
+ */
+static void SetNodeError(UEdGraphNode* Node, FText const& ErrorText)
 {
-};
+	if (ErrorText.IsEmpty())
+	{
+		Node->ErrorMsg.Empty();
+		Node->ErrorType = EMessageSeverity::Info;
+		Node->bHasCompilerMessage = false;
+	}
+	else if (Node->bHasCompilerMessage)
+	{
+		Node->ErrorMsg += TEXT("\n") + ErrorText.ToString();
+		Node->ErrorType = EMessageSeverity::Error;
+	}
+	else
+	{
+		Node->ErrorMsg = ErrorText.ToString();
+		Node->ErrorType = EMessageSeverity::Error;
+		Node->bHasCompilerMessage = true;
+	}
+}
 
-enum EVisitPhase
-{
-	VISIT_Pre,
-	VISIT_Post,
-	VISIT_Leaf
-};
+/*******************************************************************************
+ * FExpressionVisitor
+*******************************************************************************/
 
+/** 
+ * This is the base class used for IFExpressionNode tree traversal (set up to 
+ * handle different node types... new node types should have a Visit() method 
+ * added for them).
+ */
 class FExpressionVisitor
 {
 public:
-	virtual void Visit(class FExpressionNode& Node, EVisitPhase Phase) { VisitUnhandled(Node, Phase); }
-	virtual void Visit(class FTokenWrapperNode& Node, EVisitPhase Phase) { VisitUnhandled((class FExpressionNode&)Node, Phase); }
-	virtual void Visit(class FExpressionList& Node, EVisitPhase Phase) { VisitUnhandled((class FExpressionNode&)Node, Phase); }
-	virtual void Visit(class FBinaryOperator& Node, EVisitPhase Phase) { VisitUnhandled((class FExpressionNode&)Node, Phase); }
-	virtual void Visit(class FCastOperator& Node, EVisitPhase Phase) { VisitUnhandled((class FExpressionNode&)Node, Phase); }
-	virtual void Visit(class FUnaryOperator& Node, EVisitPhase Phase) { VisitUnhandled((class FExpressionNode&)Node, Phase); }
-	virtual void Visit(class FConditionalOperator& Node, EVisitPhase Phase) { VisitUnhandled((class FExpressionNode&)Node, Phase); }
+	/** 
+	 * FExpressionNodes determine when a traverser (FExpressionVisitor) has access 
+	 * to the node. There are a couple hook points, allowing the traverser to pick 
+	 * either a pre-order or post-order tree traversal. These values let the 
+	 * FExpressionVisitor where we are in the tree search.
+	 */
+	enum EVisitPhase
+	{
+		VISIT_Pre,  // the node being visited has yet to visit its children (and will next, starting with the left)
+		VISIT_Post, // the node being visited has finished visiting its children (and is about to return up, to its parent)
+		VISIT_Leaf  // the node being visited is a leaf (no children will be visited)
+	};
+
+	/**
+	 * Intended to be overridden by sub-classes, for special handling of 
+	 * explicit node types (new node types should have one added for them).
+	 * 
+	 * @param  Node		The current node in the tree's traversal.
+	 * @param  Phase	Where we currently are in traversing the specified node (before children vs. after)
+	 * @return True to continue traversing the tree, false to abort. 
+	 */
+	virtual bool Visit(class IFExpressionNode&  Node, EVisitPhase Phase)    { return VisitUnhandled(Node, Phase); }
+	virtual bool Visit(class FTokenWrapperNode& Node, EVisitPhase Phase)    { return VisitUnhandled((class IFExpressionNode&)Node, Phase); }
+	virtual bool Visit(class FBinaryOperator&   Node, EVisitPhase Phase)    { return VisitUnhandled((class IFExpressionNode&)Node, Phase); }
+	virtual bool Visit(class FUnaryOperator&    Node, EVisitPhase Phase)    { return VisitUnhandled((class IFExpressionNode&)Node, Phase); }
+	virtual bool Visit(class FConditionalOperator& Node, EVisitPhase Phase) { return VisitUnhandled((class IFExpressionNode&)Node, Phase); }
+	virtual bool Visit(class FExpressionList&   Node, EVisitPhase Phase)    { return VisitUnhandled((class IFExpressionNode&)Node, Phase); }
+	virtual bool Visit(class FFunctionExpression&  Node, EVisitPhase Phase) { return VisitUnhandled((class IFExpressionNode&)Node, Phase); }
 
 protected:
-	virtual void VisitUnhandled(class FExpressionNode& Node, EVisitPhase Phase)
+	/**
+	 * Called by all the base Visit() methods, a good point for sub-classes to 
+	 * hook into for handling EVERY IFExpressionNode type (unless they override
+	 * a Visit method)
+	 * 
+	 * @param  Node		The current node in the tree's traversal.
+	 * @param  Phase	Where we currently are in traversing the specified node (before children vs. after)	
+	 * @return True to continue traversing the tree, false to abort. 
+	 */
+	virtual bool VisitUnhandled(class IFExpressionNode& Node, EVisitPhase Phase)
 	{
+		// if we end up here, then the subclass decided not to handle the 
+		// specific node type, and therefore doesn't care about it
+		return true;
 	}
 };
 
-class FExpressionNode
+/*******************************************************************************
+ * IFExpressionNode Types
+*******************************************************************************/
+
+/**
+ * Base class for all expression-tree nodes that are generated from parsing an
+ * expression string. Represents either a single value/variable, or an operation
+ * on other IFExpressionNode(s).
+ */
+class IFExpressionNode
 {
 public:
-	virtual FString ToString() const
-	{
-		return TEXT("NOT_IMPLEMENTED");
-	}
+	/** 
+	 * Entry point for traversing the expression-tree, should either pass the 
+	 * Visitor along to sub child nodes (in the case of a branch node), or 
+	 * simply let the Visitor "visit" the leaf node.
+	 *
+	 * @return True to continue traversing the tree, false to abort. 
+	 */
+	virtual bool Accept(FExpressionVisitor& Visitor) = 0;
 
-	virtual void Accept(FExpressionVisitor& Visitor)
-	{
-		Visitor.Visit(*this, VISIT_Leaf);
-	}
+	/** 
+	 * For debug purposes, intended to help visualize what this node represents 
+	 * (for reconstructing a pseudo expression).
+	 */
+	virtual FString ToString() const = 0;
+
+protected:
+	IFExpressionNode() {}
 };
 
-class FTokenWrapperNode : public FExpressionNode
+/** 
+ * Leaf node for the expression-tree. Encapsulates either a literal constant 
+ * (FBasicToken::TOKEN_Const), or a variable identifier (FBasicToken::TOKEN_Identifier).
+ */
+class FTokenWrapperNode : public IFExpressionNode
 {
 public:
-	FTokenWrapperNode(FToken InToken)
+	FTokenWrapperNode(FBasicToken InToken)
 		: Token(InToken)
 	{
 	}
 
-	virtual void Accept(FExpressionVisitor& Visitor)
+	/** 
+	 * Gives the FExpressionVisitor access to this node (lets it "visit" this, 
+	 * in tree traversal terms).
+	 *
+	 * @return True to continue traversing the tree, false to abort. 
+	 */
+	virtual bool Accept(FExpressionVisitor& Visitor) OVERRIDE
 	{
-		Visitor.Visit(*this, VISIT_Leaf);
+		return Visitor.Visit(*this, FExpressionVisitor::VISIT_Leaf);
 	}
 
+	/** For debug purposes, constructs a textual representation of this expression */
 	virtual FString ToString() const OVERRIDE
 	{
-		if (Token.TokenType == TOKEN_Identifier)
+		if (Token.TokenType == FBasicToken::TOKEN_Identifier)
 		{
 			return FString::Printf(TEXT("identifier'%s'"), Token.Identifier);
 		}
-		else if (Token.TokenType == TOKEN_Const)
+		else if (Token.TokenType == FBasicToken::TOKEN_Const)
 		{
 			return FString::Printf(TEXT("const'%s'"), *Token.GetConstantValue());
 		}
@@ -1028,863 +238,1496 @@ public:
 	}
 
 public:
-	FToken Token;
+	/** The base token which this leaf node represents */
+	FBasicToken Token;
 };
 
 
-class FExpressionList : public FExpressionNode
+/**
+ * Branch node that represents a binary operation, where its children are the 
+ * left and right operands:
+ *
+ *                 <operator>
+ *                 /        \
+ * <left-expression>        <right-expression>
+ */
+class FBinaryOperator : public IFExpressionNode
 {
 public:
-	FExpressionList()
-	{
-	}
-
-	void Add(TSharedRef<FExpressionNode> Node)
+	FBinaryOperator(const FString& InOperator, TSharedRef<IFExpressionNode> InLHS, TSharedRef<IFExpressionNode> InRHS)
+		: Operator(InOperator)
+		, LHS(InLHS)
+		, RHS(InRHS)
 	{
-		//@TODO: Do something here, for reals
-		new (Entries) TSharedRef<FExpressionNode>(Node);
 	}
 
-	virtual FString ToString() const OVERRIDE
+	/** 
+	 * Gives the FExpressionVisitor access to this node, and passes it along to 
+	 * traverse the children.
+	 *
+	 * @return True to continue traversing the tree, false to abort. 
+	 */
+	virtual bool Accept(FExpressionVisitor& Visitor) OVERRIDE
 	{
-		FString Result;
-		for (int32 Index = 0; Index < Entries.Num(); ++Index)
+		bool bAbort = !Visitor.Visit(*this, FExpressionVisitor::VISIT_Pre);
+		if (bAbort || !LHS->Accept(Visitor) || !RHS->Accept(Visitor))
 		{
-			TSharedRef<FExpressionNode> Foo = Entries[Index];
-			Result += Foo->ToString();
+			return false;
 		}
-
-		return Result;
+		return Visitor.Visit(*this, FExpressionVisitor::VISIT_Post);
 	}
 
-	virtual void Accept(FExpressionVisitor& Visitor)
+	/** For debug purposes, constructs a textual representation of this expression */
+	virtual FString ToString() const OVERRIDE
 	{
-		Visitor.Visit(*this, VISIT_Pre);
-		for (int32 Index = 0; Index < Entries.Num(); ++Index)
-		{
-			TSharedRef<FExpressionNode> Foo = Entries[Index];
-			Foo->Accept(Visitor);
-		}
-		Visitor.Visit(*this, VISIT_Post);
+		const FString LeftStr = LHS->ToString();
+		const FString RightStr = RHS->ToString();
+		return FString::Printf(TEXT("(%s %s %s)"), *LeftStr, *Operator, *RightStr);
 	}
 
-	TArray< TSharedRef<FExpressionNode> > Entries;
+public:
+	FString Operator;
+	TSharedRef<IFExpressionNode> LHS;
+	TSharedRef<IFExpressionNode> RHS;
 };
 
-class FBinaryOperator : public FExpressionNode
+
+/**
+ * Branch node that represents a unary (prefix) operation, where its child is  
+ * the right operand:
+ *
+ *     <unary-operator>
+ *                    \
+ *                    <operand-expression>
+ */
+class FUnaryOperator : public IFExpressionNode
 {
 public:
-	FBinaryOperator(const FString& InOperator, TSharedRef<FExpressionNode> InLHS, TSharedRef<FExpressionNode> InRHS)
+	FUnaryOperator(const FString& InOperator, TSharedRef<IFExpressionNode> InRHS)
 		: Operator(InOperator)
-		, LHS(InLHS)
 		, RHS(InRHS)
 	{
 	}
 
-	virtual FString ToString() const OVERRIDE
+	/** 
+	 * Gives the FExpressionVisitor access to this node, and passes it along to 
+	 * traverse its child.
+	 *
+	 * @return True to continue traversing the tree, false to abort. 
+	 */
+	virtual bool Accept(FExpressionVisitor& Visitor) OVERRIDE
 	{
-		const FString LeftStr = LHS->ToString();
-		const FString RightStr = RHS->ToString();
-		return FString::Printf(TEXT("(%s %s %s)"), *LeftStr, *Operator, *RightStr);
+		bool bAbort = !Visitor.Visit(*this, FExpressionVisitor::VISIT_Pre);
+		if (bAbort || !RHS->Accept(Visitor))
+		{
+			return false;
+		}
+		return Visitor.Visit(*this, FExpressionVisitor::VISIT_Post);
 	}
 
-	virtual void Accept(FExpressionVisitor& Visitor)
+	/** For debug purposes, constructs a textual representation of this expression */
+	virtual FString ToString() const OVERRIDE
 	{
-		Visitor.Visit(*this, VISIT_Pre);
-		LHS->Accept(Visitor);
-		RHS->Accept(Visitor);
-		Visitor.Visit(*this, VISIT_Post);
+		const FString RightStr = RHS->ToString();
+		return FString::Printf(TEXT("(%s %s)"), *Operator, *RightStr);
 	}
+
 public:
 	FString Operator;
-	TSharedRef<FExpressionNode> LHS;
-	TSharedRef<FExpressionNode> RHS;
+	TSharedRef<IFExpressionNode> RHS;
 };
 
-class FCastOperator : public FExpressionNode
+
+/**
+ * Branch node that represents a ternary conditional (if-then-else) operation 
+ * (c ? a : b), where its children are the condition, the "then" expression,   
+ * and the "else" expression:
+ *
+ *             <conditional-operator>
+ *             /          |         \
+ *  <condition>   <then-exression>   <else-expression>
+ */
+class FConditionalOperator : public IFExpressionNode
 {
 public:
-	FCastOperator(TSharedRef<FExpressionNode> InTypeExpression, TSharedRef<FExpressionNode> InValueExpression)
-		: TypeExpression(InTypeExpression)
-		, ValueExpression(InValueExpression)
+	FConditionalOperator(TSharedRef<IFExpressionNode> InCondition, TSharedRef<IFExpressionNode> InTruePart, TSharedRef<IFExpressionNode> InFalsePart)
+		: Condition(InCondition)
+		, TruePart(InTruePart)
+		, FalsePart(InFalsePart)
 	{
 	}
 
-	virtual FString ToString() const OVERRIDE
+	/** 
+	 * Gives the FExpressionVisitor access to this node, and passes it along to 
+	 * traverse the children.
+	 *
+	 * @return True to continue traversing the tree, false to abort.  
+	 */
+	virtual bool Accept(FExpressionVisitor& Visitor) OVERRIDE
 	{
-		const FString TypeStr = TypeExpression->ToString();
-		const FString ValueStr = ValueExpression->ToString();
-		return FString::Printf(TEXT("((%s)(%s))"), *TypeStr, *ValueStr);
+		bool bAbort = !Visitor.Visit(*this, FExpressionVisitor::VISIT_Pre);
+		// @TODO: what about the Condition?
+		if (bAbort || !TruePart->Accept(Visitor) || !FalsePart->Accept(Visitor))
+		{
+			return false;
+		}
+		return Visitor.Visit(*this, FExpressionVisitor::VISIT_Post);
 	}
 
-	virtual void Accept(FExpressionVisitor& Visitor)
+	/** For debug purposes, constructs a textual representation of this expression */
+	virtual FString ToString() const OVERRIDE
 	{
-		Visitor.Visit(*this, VISIT_Pre);
-		TypeExpression->Accept(Visitor);
-		ValueExpression->Accept(Visitor);
-		Visitor.Visit(*this, VISIT_Post);
+		const FString ConditionStr = Condition->ToString();
+		const FString TrueStr = TruePart->ToString();
+		const FString FalseStr = FalsePart->ToString();
+		return FString::Printf(TEXT("(%s ? %s : %s)"), *ConditionStr, *TrueStr, *FalseStr);
 	}
 
 public:
-	TSharedRef<FExpressionNode> TypeExpression;
-	TSharedRef<FExpressionNode> ValueExpression;
+	TSharedRef<IFExpressionNode> Condition;
+	TSharedRef<IFExpressionNode> TruePart;
+	TSharedRef<IFExpressionNode> FalsePart;
 };
 
-class FUnaryOperator : public FExpressionNode
+/**
+ * Branch node that represents an n-dimentional list of sub-expressions (like 
+ * for vector parameter lists, etc.). Each child is a separate sub-expression:
+ *
+ *                 <list-node>
+ *                 /    |    \
+ * <sub-expression0>    |    <sub-expression2>
+ *                      |
+ *              <sub-expression1>              
+ */
+class FExpressionList : public IFExpressionNode
 {
 public:
-	FUnaryOperator(const FString& InOperator, TSharedRef<FExpressionNode> InRHS)
-		: Operator(InOperator)
-		, RHS(InRHS)
+	/** 
+	 * Gives the FExpressionVisitor access to this node, and passes it along to 
+	 * traverse all children.
+	 *
+	 * @return True to continue traversing the tree, false to abort. 
+	 */
+	virtual bool Accept(FExpressionVisitor& Visitor) OVERRIDE
 	{
+		bool bAbort = !Visitor.Visit(*this, FExpressionVisitor::VISIT_Pre);
+		for (TSharedRef<IFExpressionNode> Child : Children)
+		{
+			if (bAbort || !Child->Accept(Visitor))
+			{
+				return false;
+			}
+		}
+		return Visitor.Visit(*this, FExpressionVisitor::VISIT_Post);
 	}
-
+    
+	/** For debug purposes, constructs a textual representation of this expression */
 	virtual FString ToString() const OVERRIDE
 	{
-		const FString RightStr = RHS->ToString();
-		return FString::Printf(TEXT("(%s %s)"), *Operator, *RightStr);
+		FString AsString("(");
+		for (TSharedRef<IFExpressionNode> Child : Children)
+		{
+			AsString += Child->ToString();
+			if (Child == Children.Last())
+			{
+				AsString += ")";
+			}
+			else 
+			{
+				AsString += ", ";
+			}
+		}
+		return AsString;
+	}
+    
+public:
+	TArray< TSharedRef<IFExpressionNode> > Children;
+};
+
+/**
+ * Branch node that represents some function (like sin(), cos(), etc.), could 
+ * also represent some structure (conceptually the constructor), like vector, 
+ * rotator, etc. Its child is a single FExpressionList (which wraps all the params).         
+ */
+class FFunctionExpression : public IFExpressionNode
+{
+public:
+	FFunctionExpression(FString const& InFuncName, TSharedRef<FExpressionList> InParamList)
+		: FuncName(InFuncName)
+		, ParamList(InParamList)
+	{
+	}
+
+	/** 
+	 * Gives the FExpressionVisitor access to this node, and passes it along to 
+	 * traverse its child.
+	 *
+	 * @return True to continue traversing the tree, false to abort. 
+	 */
+	virtual bool Accept(FExpressionVisitor& Visitor) OVERRIDE
+	{
+		bool bAbort = !Visitor.Visit(*this, FExpressionVisitor::VISIT_Pre);
+		if (bAbort || !ParamList->Accept(Visitor))
+		{
+			return false;
+		}
+		return Visitor.Visit(*this, FExpressionVisitor::VISIT_Post);
 	}
 
-	virtual void Accept(FExpressionVisitor& Visitor)
+	/** For debug purposes, constructs a textual representation of this expression */
+	virtual FString ToString() const OVERRIDE
 	{
-		Visitor.Visit(*this, VISIT_Pre);
-		RHS->Accept(Visitor);
-		Visitor.Visit(*this, VISIT_Post);
+		FString const ParamsString = ParamList->ToString();
+		return FString::Printf(TEXT("(%s%s)"), *FuncName, *ParamsString);
 	}
+
 public:
-	FString Operator;
-	TSharedRef<FExpressionNode> RHS;
+	FString FuncName;
+	TSharedRef<FExpressionList> ParamList;
 };
 
-/////////////////////////////////////////////////////
-
-class FConditionalOperator : public FExpressionNode
+/*******************************************************************************
+ * FLayoutVisitor
+*******************************************************************************/
+
+/**
+ * This class is utilized to help layout math expression nodes by traversing the
+ * expression tree and cataloging each expression node's depth. From the tree's 
+ * depth we can determine the width of the the graph (an where to place each K2 node):
+ *
+ *    _
+ *   |            [_]---[_]
+ *   |           /
+ * height   [_]--       [_]--[_]---[_]
+ *   |           \     /
+ *   |_           [_]---[_]
+ *
+ *		    ^-------depth/width-------^
+ */
+class FLayoutVisitor : public FExpressionVisitor
 {
 public:
-	FConditionalOperator(TSharedRef<FExpressionNode> InCondition, TSharedRef<FExpressionNode> InTruePart, TSharedRef<FExpressionNode> InFalsePart)
-		: Condition(InCondition)
-		, TruePart(InTruePart)
-		, FalsePart(InFalsePart)
+	/** Tracks the horizontal (depth) placement for each expression node encountered */
+	TMap<IFExpressionNode*, int32> DepthChart;
+	/** Tracks the vertical (height) placement for each expression node encountered */
+	TMap<IFExpressionNode*, int32> HeightChart;
+	/** Tracks the total height (value) at each depth (key) */
+	TMap<int32, int32> DepthHeightLookup;
+
+	/** */
+	FLayoutVisitor()
+		: CurrentDepth(0)
+		, MaximumDepth(0)
 	{
 	}
 
-	virtual FString ToString() const OVERRIDE
+	/**
+	 * Retrieves the total depth (or graph width) of the previously traversed
+	 * expression tree.
+	 */
+	int32 GetMaximumDepth() const
 	{
-		const FString ConditionStr = Condition->ToString();
-		const FString TrueStr = TruePart->ToString();
-		const FString FalseStr = FalsePart->ToString();
-		return FString::Printf(TEXT("(%s ? %s : %s)"), *ConditionStr, *TrueStr, *FalseStr);
+		return MaximumDepth;
 	}
 
-	virtual void Accept(FExpressionVisitor& Visitor)
+	/**
+	 * Resets this tree visitor so that it can accurately parse another 
+	 * expression tree (else the results would stack up).
+	 */
+	void Clear() 
 	{
-		Visitor.Visit(*this, VISIT_Pre);
-		TruePart->Accept(Visitor);
-		FalsePart->Accept(Visitor);
-		Visitor.Visit(*this, VISIT_Post);
+		CurrentDepth = 0;
+		MaximumDepth = 0;
+		DepthChart.Empty();
+		HeightChart.Empty();
+		DepthHeightLookup.Empty();
 	}
-public:
-	TSharedRef<FExpressionNode> Condition;
-	TSharedRef<FExpressionNode> TruePart;
-	TSharedRef<FExpressionNode> FalsePart;
+
+private:
+	/**
+	 * From the FExpressionVisitor interface, a generic choke point for visiting 
+	 * all expression nodes.
+	 *
+	 * @return True to continue traversing the tree, false to abort. 
+	 */
+	virtual bool VisitUnhandled(class IFExpressionNode& Node, EVisitPhase Phase) OVERRIDE
+	{
+		if (Phase == FExpressionVisitor::VISIT_Pre)
+		{
+			++CurrentDepth;
+			MaximumDepth = FMath::Max(CurrentDepth, MaximumDepth);
+		}
+		else
+		{
+			if (Phase == FExpressionVisitor::VISIT_Post)
+			{
+				--CurrentDepth;
+			}
+			// else leaf
+
+			// CurrentHeight represents how many nodes have already been placed 
+			// at this depth
+			int32& CurrentHeight = DepthHeightLookup.FindOrAdd(CurrentDepth);
+
+			DepthChart.Add(&Node, CurrentDepth);
+			HeightChart.Add(&Node, CurrentHeight);
+
+			// since we just placed another node at this depth, increase the 
+			// height count
+			++CurrentHeight;
+		}
+
+		// let the tree traversal continue! don't abort it!
+		return true;
+	}
+
+private:
+	int32 CurrentDepth;
+	int32 MaximumDepth;
 };
 
-/////////////////////////////////////////////////////
-// 
+/*******************************************************************************
+ * FOperatorTable
+*******************************************************************************/
 
-struct FOperatorTable
+/** 
+ * This class acts as a lookup table for mapping operator strings (like "+", 
+ * "*", etc.) to corresponding functions that can be turned into blueprint 
+ * nodes. It builds itself (so users don't have to add mappings themselves).
+ */
+class FOperatorTable
 {
-	TArray<UFunction*> MatchingFunctions;
+public:
+	FOperatorTable() { Rebuild(); }
 
-	UFunction* FindFunction(const TArray<FEdGraphPinType>& DesiredTypes) const
+	/** 
+	 * Checks to see if there are any functions associated with the specified 
+	 * operator. 
+	 */
+	bool Contains(FString const& Operator) const
 	{
-		const UEdGraphSchema_K2* K2Schema = GetDefault<UEdGraphSchema_K2>();
+		return LookupTable.Contains(Operator);
+	}	
+
+	/**
+	 * Attempts to lookup a function matching the supplied signature (where 
+	 * 'Operator' identifies the function's name and 'InputTypeList' defines
+	 * the desired parameters). If one can't be found, it attempts to find a
+	 * match by promoting the input types (like from int to float, etc.)
+	 * 
+	 * @param  Operator			The operator you want to find a function for.
+	 * @param  InputTypeList	A list of parameter types you want to feed the function.
+	 * @return A pointer to the matching function (if one was found), otherwise nullptr.
+	 */
+	UFunction* FindMatchingFunction(FString const& Operator, TArray<FEdGraphPinType> const& InputTypeList) const
+	{
+		// make a local copy of the desired input types so that we can promote 
+		// those types as needed
+		TArray<FEdGraphPinType> ParamTypeList = InputTypeList;
+
+		// try to find the function
+		UFunction* MatchingFunc = FindFunctionInternal(Operator, ParamTypeList);
 
-		for (int32 FuncIndex = 0; FuncIndex < MatchingFunctions.Num(); ++FuncIndex)
+		// if we didn't find a function that matches the supplied function 
+		// signature, then try to promote the parameters (like from int to 
+		// float), and see if we can lookup a function with those types
+		for (int32 promoterIndex = 0; (promoterIndex < OrderedTypePromoters.Num()) && (MatchingFunc == NULL); ++promoterIndex)
 		{
-			UFunction* TestFunction = MatchingFunctions[FuncIndex];
+			FTypePromoter const& PromotionOperator = OrderedTypePromoters[promoterIndex];
 
-			
-			int32 ArgumentIndex = 0;
-			for (TFieldIterator<UProperty> PropIt(TestFunction); PropIt && (PropIt->PropertyFlags & CPF_Parm); ++PropIt)
+			// Apply the promotion operator to any values that match
+			bool bMadeChanges = false;
+			for (FEdGraphPinType& ParamType : ParamTypeList)
+			{
+				bMadeChanges |= PromotionOperator.Execute(ParamType);
+			}
+
+			// since we've promoted some of the params, attempt to find the 
+			// function again (maybe there's one that matches these param types)
+			if (bMadeChanges)
 			{
-				UProperty* Param = *PropIt;
+				MatchingFunc = FindFunctionInternal(Operator, ParamTypeList);
+				// if we found a function to match this time around, no need to 
+				// continue with 
+				if (MatchingFunc != NULL)
+				{
+					break;
+				}
+			}
+		}
+
+		return MatchingFunc;
+	}
+
+	/** 
+	 * Flags the specified function as one associated with the supplied 
+	 * operator.
+	 */
+	void Add(FString const& Operator, UFunction* OperatorFunc)
+	{
+		LookupTable.FindOrAdd(Operator).Add(OperatorFunc);
+	}
+
+	/**
+	 * Rebuilds the lookup table, mapping operator strings (like "+" or "*") to
+	 * associated functions (searches through function libraries for operator 
+	 * functions).
+	 */
+	void Rebuild()
+	{
+		LookupTable.Empty();
+		OrderedTypePromoters.Empty();
 
-				if (!Param->HasAnyPropertyFlags(CPF_ReturnParm))
+		// run through all blueprint function libraries and build up a list of 
+		// functions that have good operator info
+		for (TObjectIterator<UClass> ClassIt; ClassIt; ++ClassIt)
+		{
+			UClass* TestClass = *ClassIt;
+			if (TestClass->IsChildOf(UBlueprintFunctionLibrary::StaticClass()) && (!TestClass->HasAnyClassFlags(CLASS_Abstract)))
+			{
+				for (TFieldIterator<UFunction> FuncIt(TestClass, EFieldIteratorFlags::ExcludeSuper); FuncIt; ++FuncIt)
 				{
-					if (ArgumentIndex < DesiredTypes.Num())
+					UFunction* TestFunction = *FuncIt;
+					
+					if (!TestFunction->HasAnyFunctionFlags(FUNC_BlueprintPure) || (TestFunction->GetReturnProperty() == nullptr))
+					{
+						continue;
+					}
+					
+					FString FunctionName = TestFunction->GetName();
+					TArray<FString> const& OperatorAliases = GetOperatorAliases(FunctionName);
+					
+					// if there are aliases, use those instead of the function's standard name
+					if (OperatorAliases.Num() > 0)
 					{
-						FEdGraphPinType ParamType;
-						if (K2Schema->ConvertPropertyToPinType(Param, /*out*/ ParamType))
+						for (FString const& Alias : OperatorAliases)
 						{
-							const FEdGraphPinType& TypeToMatch = DesiredTypes[ArgumentIndex];
+							Add(Alias, TestFunction);
+						}
+					}
+					else if (TestFunction->HasMetaData(FBlueprintMetadata::MD_CompactNodeTitle))
+					{
+						FunctionName = TestFunction->GetMetaData(FBlueprintMetadata::MD_CompactNodeTitle);
+						Add(FunctionName, TestFunction);
+					}
+					else if (TestFunction->HasMetaData(FBlueprintMetadata::MD_FriendlyName))
+					{
+						FunctionName = TestFunction->GetMetaData(FBlueprintMetadata::MD_FriendlyName);
+						Add(FunctionName, TestFunction);
+					}
+				}
+			}
+		}
 
-							if (!K2Schema->ArePinTypesCompatible(TypeToMatch, ParamType))
+		FTypePromoter ByteToIntPromoter;
+		ByteToIntPromoter.BindStatic(&PromoteByteToInt);
+		OrderedTypePromoters.Add(ByteToIntPromoter);
+
+		FTypePromoter IntToFloatPromoter;
+		IntToFloatPromoter.BindStatic(&PromoteIntToFloat);
+		OrderedTypePromoters.Add(IntToFloatPromoter);
+	}
+
+private:
+	/**
+	 * Attempts to lookup a function matching the supplied signature (where 
+	 * 'Operator' identifies the function's name and 'InputTypeList' defines
+	 * the desired parameters into that function).
+	 * 
+	 * @param  Operator			The operator you want to find a function for.
+	 * @param  InputTypeList	A list of parameter types you want to feed the function.
+	 * @return A pointer to the matching function (if one was found), otherwise nullptr.
+	 */
+	UFunction* FindFunctionInternal(FString const& Operator, TArray<FEdGraphPinType> const& InputTypeList) const
+	{
+		UFunction* MatchedFunction = nullptr;
+
+		FFunctionsList const* OperatorFunctions = LookupTable.Find(Operator);
+		if (OperatorFunctions != nullptr)
+		{
+			UEdGraphSchema_K2 const* K2Schema = GetDefault<UEdGraphSchema_K2>();
+			for (UFunction* TestFunction : *OperatorFunctions)
+			{
+				int32 ArgumentIndex = 0;
+				for (TFieldIterator<UProperty> PropIt(TestFunction); PropIt && (PropIt->PropertyFlags & CPF_Parm); ++PropIt)
+				{
+					UProperty* Param = *PropIt;
+					if (!Param->HasAnyPropertyFlags(CPF_ReturnParm))
+					{
+						if (ArgumentIndex < InputTypeList.Num())
+						{
+							FEdGraphPinType ParamType;
+							if (K2Schema->ConvertPropertyToPinType(Param, /*out*/ParamType))
 							{
-								// Type mismatch
-								break;
+								FEdGraphPinType const& TypeToMatch = InputTypeList[ArgumentIndex];
+								if (!K2Schema->ArePinTypesCompatible(TypeToMatch, ParamType))
+								{
+									break; // type mismatch
+								}
+							}
+							else
+							{
+								break; // function has a non-K2 type as a parameter
 							}
 						}
 						else
 						{
-							// Function has a non-K2 type as a parameter
-							break;
+							break; // ran out of arguments; no match
 						}
+						++ArgumentIndex;
 					}
-					else
-					{
-						// Ran out of arguments; no match
-						break;
-					}
+				}
 
-					++ArgumentIndex;
+				if (ArgumentIndex == InputTypeList.Num())
+				{
+					// success!
+					MatchedFunction = TestFunction;
+					break;
 				}
 			}
+		}
+
+		return MatchedFunction;
+	}
+	
+	/**
+	 * Here we overwrite and map multiplies names to specific functions (for 
+	 * example "MultiplyMultiply_FloatFloat" and "^2" are not the sort of names
+	 * we'd expect a user to input in a mathematical expression). We can 
+	 * replace a function name with a single value, or a series of values 
+	 * (could setup "asin" and "arcsin" both as aliases for the ASin() method).
+	 *
+	 * @param  FunctionName		The raw name of the function you're looking to replace (not the friendly name)
+	 * @return A reference to the array of aliases for the specified function (an empty array if none were found).
+	 */
+	static TArray<FString> const& GetOperatorAliases(FString const& FunctionName)
+	{
+#define FUNC_ALIASES_BEGIN(FuncName) \
+		if (FunctionName == FString(TEXT(FuncName))) \
+		{ \
+			static TArray<FString> AliasTable; \
+			if (AliasTable.Num() == 0) \
+			{
+#define ADD_ALIAS(AliasStr) \
+				AliasTable.Add(TEXT(AliasStr));
+#define FUNC_ALIASES_END \
+			} \
+			return AliasTable; \
+		}
+				
+		FUNC_ALIASES_BEGIN("BooleanAND")
+			ADD_ALIAS("&&")
+		FUNC_ALIASES_END
+				
+		FUNC_ALIASES_BEGIN("BooleanOR")
+			ADD_ALIAS("||")
+		FUNC_ALIASES_END
+				
+		FUNC_ALIASES_BEGIN("BooleanXOR")
+			ADD_ALIAS("^")
+		FUNC_ALIASES_END
+				
+		FUNC_ALIASES_BEGIN("Not_PreBool")
+			ADD_ALIAS("!")
+		FUNC_ALIASES_END
+		
+		// keep the compact node title of "^2" from being the required key
+		FUNC_ALIASES_BEGIN("Square")
+			ADD_ALIAS("SQUARE")
+		FUNC_ALIASES_END
+				
+		FUNC_ALIASES_BEGIN("MultiplyMultiply_FloatFloat")
+			ADD_ALIAS("POWER")
+			ADD_ALIAS("POW")
+		FUNC_ALIASES_END
+				
+		FUNC_ALIASES_BEGIN("ASin")
+			// have to add "ASin" back, because this overwrites the function's
+			// name and we still want it as a viable option
+			ADD_ALIAS("ASIN")
+			ADD_ALIAS("ARCSIN")
+		FUNC_ALIASES_END
+		
+		FUNC_ALIASES_BEGIN("MakeVector")
+			ADD_ALIAS("VECTOR")
+			ADD_ALIAS("VEC")
+			ADD_ALIAS("VECT")
+		FUNC_ALIASES_END
+		
+		FUNC_ALIASES_BEGIN("MakeVector2D")
+			ADD_ALIAS("VECTOR2D")
+			ADD_ALIAS("VEC2D")
+			ADD_ALIAS("VECT2D")
+		FUNC_ALIASES_END
+	
+		FUNC_ALIASES_BEGIN("MakeRot")
+			ADD_ALIAS("ROTATOR")
+			ADD_ALIAS("ROT")
+		FUNC_ALIASES_END
+		
+		FUNC_ALIASES_BEGIN("MakeTransform")
+			ADD_ALIAS("TRANSFORM")
+			ADD_ALIAS("XFORM")
+		FUNC_ALIASES_END
+		
+		FUNC_ALIASES_BEGIN("MakeColor")
+			ADD_ALIAS("COLOR")
+			ADD_ALIAS("LINEARCOLOR")
+			ADD_ALIAS("COLOUR") // long live the empire!
+		FUNC_ALIASES_END
+		
+		FUNC_ALIASES_BEGIN("RandomFloat")
+			ADD_ALIAS("RandomFloat")
+			ADD_ALIAS("RAND")
+			ADD_ALIAS("RANDOM")
+		FUNC_ALIASES_END
+				
+		// if none of the above aliases returned, then we don't have any for
+		// this function (use its regular name)
+		static TArray<FString> NoAliases;
+		return NoAliases;
+				
+#undef FUNC_ALIASES_END
+#undef ADD_ALIAS
+#undef FUNC_ALIASES_END
+	}
 
-			if (ArgumentIndex == DesiredTypes.Num())
-			{
-				// Success!
-				return TestFunction;
-			}
-		}
+private:
+	/** 
+	 * A single operator can have multiple functions associated with it; usually 
+	 * for handling different types (int*int, vs. int*vector), hence this array.
+	 */
+	typedef TArray<UFunction*> FFunctionsList;
+	/** 
+	 * A lookup table, mapping operator strings (like "+", "*", etc.) to a list 
+	 * of associated functions. 
+	 */
+	TMap<FString, FFunctionsList> LookupTable;
 
-		return NULL;
-	}
+	/**
+	 * When looking to match parameters, there are some implicit conversions we
+	 * can make to try and find a match (like converting from int to float). 
+	 * This holds an ordered list of delegates that will try and promote the 
+	 * supplied types.
+	 */
+	DECLARE_DELEGATE_RetVal_OneParam(bool, FTypePromoter, FEdGraphPinType&);
+	TArray<FTypePromoter> OrderedTypePromoters;
 };
 
-/////////////////////////////////////////////////////
-
+/*******************************************************************************
+ * FCodeGenFragments
+*******************************************************************************/
+
+/**
+ * FCodeGenFragments facilitate the making of pin connections/defaults. When 
+ * turning an expression tree into a network of UK2Nodes you traverse the tree,
+ * working backwards from the expression's result node. This means that when you
+ * spawn a UK2Node, you don't have the node (or literals) that should be plugged 
+ * into it, that is why these fragments are created (to track the associated 
+ * UK2Node/literal, and provide an easy interface for connecting it later with
+ * other fragments/nodes).
+ */
 class FCodeGenFragment
 {
 public:
-	FEdGraphPinType Type;
-	bool bSucceeded;
-
-	//@TODO: Display this error message somewhere
-	FString ErrorMessage;
+	FCodeGenFragment(FEdGraphPinType InType)
+		: FragmentType(InType)
+	{
+	}
 
-	virtual void ConnectToInput(UEdGraphPin* InputPin) { }
+    /**
+     * Takes the input to some other fragment, and plugs the result of this one 
+	 * into it.
+     *
+     * @param  InputPin	 Either an input into some parent expression, or the final output for the entire math expression.
+	 * @return True is the connection was made, otherwise false.
+     */
+	virtual bool ConnectToInput(UEdGraphPin* InputPin, FCompilerResultsLog& MessageLog) = 0;
 
-	FCodeGenFragment()
-		: bSucceeded(false)
+	/**
+	 * As it stands, all the math nodes/literals that can be generated have a 
+	 * singular output (hence why we have a basic "connect this fragment to an  
+	 * input" function). This message retrieves that output type.
+	 * 
+	 * @return The pin type of this fragment's output.
+	 */
+	FEdGraphPinType const& GetOutputType() const
 	{
+		return FragmentType;
 	}
 
 protected:
-	// Tries to connect two pins, verifying the type/etc..., and reporting an error if there is one
-	void SafeConnectPins(UEdGraphPin* OutputPin, UEdGraphPin* InputPin)
+	/**
+	 * Utility method for sub-classes to use when attempting a connection 
+	 * between two pins. Tries to connect two pins, verifying the type/etc, and
+	 * reporting a failure if there is one.
+	 * 
+	 * @param  OutputPin	The output pin (probably from this fragment).
+	 * @param  InputPin		The input pin (probably from some other fragment). 
+	 * @return True if the connection was made, false if the pins weren't compatible.
+	 */
+	bool SafeConnectPins(UEdGraphPin* OutputPin, UEdGraphPin* InputPin, FCompilerResultsLog& MessageLog)
 	{
 		const UEdGraphSchema* Schema = InputPin->GetSchema();
+		bool bSuccess = Schema->TryCreateConnection(OutputPin, InputPin);
 
-		if (!Schema->TryCreateConnection(OutputPin, InputPin))
+		if (!bSuccess)
 		{
-			//@TODO: Need a much better error message!
-			ErrorMessage = "Variable type is not compatible with input pin";
+			MessageLog.Error(*LOCTEXT("PinsNotCompatible", "Output pin '@@ 'is not compatible with input: '@@'").ToString(),
+				OutputPin, InputPin);
 		}
+
+		return bSuccess;
 	}
+
+private:
+	/** All fragments have a singular output type, this is considered the fragment's type as a whole. */
+	FEdGraphPinType FragmentType;
 };
 
+/**
+ * If the user uses a variable name that already exists in the blueprint, then
+ * we use that instead of adding an extra input. This fragment wraps a
+ * VariableGet node that was generated in that scenario.
+ */
 class FCodeGenFragment_VariableGet : public FCodeGenFragment
 {
-protected:
-	UK2Node_VariableGet* Node;
 public:
-	virtual void ConnectToInput(UEdGraphPin* InputPin) OVERRIDE
+	FCodeGenFragment_VariableGet(UK2Node_VariableGet* InNode, FEdGraphPinType const& InType)
+		: FCodeGenFragment(InType)
+		, GeneratedNode(InNode)
+	{
+		check(GeneratedNode != nullptr);
+	}
+
+	/// Begin FCodeGenFragment Interface
+	virtual bool ConnectToInput(UEdGraphPin* InputPin, FCompilerResultsLog& MessageLog) OVERRIDE
 	{
-		if (UEdGraphPin* VariablePin = Node->FindPin(Node->VariableReference.GetMemberName().ToString()))
+		bool bSuccess = false;
+		if (UEdGraphPin* VariablePin = GeneratedNode->FindPin(GeneratedNode->VariableReference.GetMemberName().ToString()))
 		{
-			SafeConnectPins(VariablePin, InputPin);
+			bSuccess = SafeConnectPins(VariablePin, InputPin, MessageLog);
 		}
 		else
 		{
-			ErrorMessage = "Failed to find variable output pin";
+			FText ErrorText = FText::Format(LOCTEXT("NoVariablePin", "Failed to find the '{0}' pin for: '@@'"),
+				FText::FromName(GeneratedNode->VariableReference.GetMemberName()));
+			MessageLog.Error(*ErrorText.ToString(), GeneratedNode);
 		}
+		return bSuccess;
 	}
+	/// End FCodeGenFragment Interface
 
-	FCodeGenFragment_VariableGet(UK2Node_VariableGet* InNode, const FEdGraphPinType& InType)
-		: Node(InNode)
-	{
-		check(Node);
-		Type = InType;
-		bSucceeded = true;
-	}
+private:
+	UK2Node_VariableGet* GeneratedNode;
 };
 
-
+/**
+ * All operators in the mathematical expression correspond to library functions, 
+ * which in turn generate CallFunction nodes. This fragment wraps one of those 
+ * operation nodes and connects it with the given input (when prompted to).
+ */
 class FCodeGenFragment_FuntionCall : public FCodeGenFragment
 {
-protected:
-	UK2Node_CallFunction* Node;
 public:
-	virtual void ConnectToInput(UEdGraphPin* InputPin) OVERRIDE
+	FCodeGenFragment_FuntionCall(UK2Node_CallFunction* InNode, FEdGraphPinType const& InType)
+		: FCodeGenFragment(InType)
+		, GeneratedNode(InNode)
+	{
+		check(GeneratedNode != nullptr);
+	}
+
+	/// Begin FCodeGenFragment Interface
+	virtual bool ConnectToInput(UEdGraphPin* InputPin, FCompilerResultsLog& MessageLog) OVERRIDE
 	{
-		if (UEdGraphPin* ResultPin = Node->GetReturnValuePin())
+		bool bSuccess = false;
+		if (UEdGraphPin* ResultPin = GeneratedNode->GetReturnValuePin())
 		{
-			SafeConnectPins(ResultPin, InputPin);
+			bSuccess = SafeConnectPins(ResultPin, InputPin, MessageLog);
 		}
 		else
 		{
-			ErrorMessage = "Failed to find function output pin";
+			MessageLog.Error(*LOCTEXT("NoRetValPin", "Failed to find an output pin for: '@@'").ToString(),
+				GeneratedNode);
 		}
+		return bSuccess;
 	}
+	/// End FCodeGenFragment Interface
 
-	FCodeGenFragment_FuntionCall(UK2Node_CallFunction* InNode, const FEdGraphPinType& InType)
-		: Node(InNode)
-	{
-		Type = InType;
-		bSucceeded = true;
-	}
+private:
+	UK2Node_CallFunction* GeneratedNode;
 };
 
+/**
+ * This fragment doesn't have a corresponding UK2Node, instead it represents a
+ * constant value that should be entered into another node's input field. When
+ * "connected", it modifies the connecting pin's DefaultValue.
+ */
 class FCodeGenFragment_Literal : public FCodeGenFragment
 {
 public:
-	FString DefaultValue;
+	FCodeGenFragment_Literal(FString const& LiteralVal, FEdGraphPinType const& ResultType) 
+		: FCodeGenFragment(ResultType) 
+		, DefaultValue(LiteralVal)
+	{}
 
-	virtual void ConnectToInput(UEdGraphPin* InputPin) OVERRIDE
+	/// Begin FCodeGenFragment Interface
+	virtual bool ConnectToInput(UEdGraphPin* InputPin, FCompilerResultsLog& MessageLog) OVERRIDE
 	{
-		//@TODO: Check the type!
-		//@TODO: Should we support pin object literals?
-		InputPin->DefaultValue = DefaultValue;
+		UEdGraphSchema_K2 const* K2Schema = Cast<UEdGraphSchema_K2>(InputPin->GetSchema());
+		bool bSuccess = true;//K2Schema->ArePinTypesCompatible(GetOutputType(), InputPin->PinType);
+		if (bSuccess)
+		{
+			InputPin->DefaultValue = DefaultValue;
+		}
+		else 
+		{
+			FText ErrorText = FText::Format(LOCTEXT("LiteralNotCompatible", "Literal type ({0}) is incompatible with pin: '@@'"),
+				FText::FromString(GetOutputType().PinCategory));
+			MessageLog.Error(*ErrorText.ToString(), InputPin);
+		}
+		return bSuccess;
 	}
+	/// End FCodeGenFragment Interface
+
+private: 
+	FString DefaultValue;
 };
 
+/**
+ * This fragment corresponds to an input pin that was added to the 
+ * MathExpression node. Input pins are generated when the user enters variable 
+ * names (like "x", or "y"... ones that aren't variables on the blueprint).
+ */
 class FCodeGenFragment_InputPin : public FCodeGenFragment
 {
 public:
-	virtual void ConnectToInput(UEdGraphPin* InputPin) OVERRIDE
+	FCodeGenFragment_InputPin(UEdGraphPin* InTunnelInputPin)
+		: FCodeGenFragment(InTunnelInputPin->PinType)
+		, TunnelInputPin(InTunnelInputPin)
 	{
-		SafeConnectPins(TunnelInputPin, InputPin);
 	}
 
-	FCodeGenFragment_InputPin(UEdGraphPin* InTunnelInputPin)
-		: TunnelInputPin(InTunnelInputPin)
+	/// Begin FCodeGenFragment Interface
+	virtual bool ConnectToInput(UEdGraphPin* InputPin, FCompilerResultsLog& MessageLog) OVERRIDE
 	{
-		Type = TunnelInputPin->PinType;
-		bSucceeded = true;
+		return SafeConnectPins(TunnelInputPin, InputPin, MessageLog);
 	}
+	/// End FCodeGenFragment Interface
+	
 private:
 	UEdGraphPin* TunnelInputPin;
 };
 
-/////////////////////////////////////////////////////
-// FDepthLabelVisitor
+/*******************************************************************************
+ * FMathGraphGenerator
+*******************************************************************************/
 
-// This class traverses the AST to label the depth/generation of each
-// AST node, allowing a pretty K2 node layout in the collapsed graph
-class FDepthLabelVisitor : public FExpressionVisitor
+/**
+ * Takes the root of an expression tree and instantiates blueprint nodes/pins   
+ * for the specified UK2Node_MathExpression (which is a tunnel node, similar to 
+ * how collapsed composite nodes work).
+ */
+class FMathGraphGenerator : public FExpressionVisitor
 {
 public:
-	TMap<FExpressionNode*, int32> Chart;
-
-	FDepthLabelVisitor()
-		: CurrentDepth(0)
-		, MaximumDepth(0)
+	FMathGraphGenerator(UK2Node_MathExpression* InNode)
+		: CompilingNode(InNode)
+		, TargetBlueprint(FBlueprintEditorUtils::FindBlueprintForGraphChecked(InNode->BoundGraph))
+		, ActiveMessageLog(nullptr)
 	{
 	}
 
-	int32 GetMaximumDepth() const
+	/**
+	 * Takes an expression tree and converts expression nodes into UK2Nodes, 
+	 * connecting them, and adding them under the math expression node that
+     * this was instantiated with.
+	 * 
+	 * @param  ExpressionRoot   The root of the expression tree that you want converted into a UK2Node network.
+	 */
+	bool GenerateCode(TSharedRef<IFExpressionNode> ExpressionRoot, FCompilerResultsLog& MessageLog)
 	{
-		return MaximumDepth;
-	}
-protected:
-	int32 CurrentDepth;
-	int32 MaximumDepth;
+		ActiveMessageLog = &MessageLog;
+		// want to track if we generated any errors from this pass, so we need to know how many we started with
+		int32 StartingErrorCount = MessageLog.NumErrors;
+		
+		LayoutMapper.Clear();
+		// map the depth/height of expression tree (so we can position nodes prettily)
+		ExpressionRoot->Accept(LayoutMapper);
+		// reset the bounds tracking, so we can adjust it as we spawn nodes
+		GraphXBounds.X = +LayoutMapper.GetMaximumDepth();
+		GraphXBounds.Y = -LayoutMapper.GetMaximumDepth();
+		
+		// traverse the expression tree, spawning nodes as we go along
+		ExpressionRoot->Accept(*this);
 
-	virtual void VisitUnhandled(class FExpressionNode& Node, EVisitPhase Phase) OVERRIDE
-	{
-		if (Phase == EVisitPhase::VISIT_Pre)
+		UK2Node_Tunnel* EntryNode = CompilingNode->GetEntryNode();
+		UK2Node_Tunnel* ExitNode  = CompilingNode->GetExitNode();
+
+		TSharedPtr<FCodeGenFragment> RootFragment = CompiledFragments.FindRef(&ExpressionRoot.Get());
+		if (RootFragment.IsValid())
 		{
-			++CurrentDepth;
-			MaximumDepth = FMath::Max(CurrentDepth, MaximumDepth);
+			// connect the final node of the expression with the math-node's output
+			UEdGraphPin* ReturnPin = ExitNode->CreateUserDefinedPin(TEXT("ReturnValue"), RootFragment->GetOutputType());
+			if (!RootFragment->ConnectToInput(ReturnPin, MessageLog))
+			{
+				MessageLog.Error(*LOCTEXT("ResultConnectError", "Failed to connect the generated nodes with expression's result pin: '@@'").ToString(),
+					ReturnPin);
+			}
 		}
 		else
 		{
-			if (Phase == EVisitPhase::VISIT_Post)
-			{
-				--CurrentDepth;
-			}
-			Chart.Add(&Node, CurrentDepth);
+			MessageLog.Error(*LOCTEXT("NoGraphGenerated", "No root node generated from the expression: '@@'").ToString(),
+				CompilingNode);
 		}
-	}
-};
-
-/////////////////////////////////////////////////////
-// FTypePromotion
 
-struct FTypePromotion
-{
-	virtual bool ApplyPotentialPromotion(FEdGraphPinType& InOutType) { return false; }
-};
+		// position the entry and exit nodes somewhere sane
+		{
+			const FVector2D EntryPos = GetNodePosition(GraphXBounds.X - 1, 0);
+			EntryNode->NodePosX = EntryPos.X;
+			EntryNode->NodePosY = EntryPos.Y;
 
-struct FTypePromotion_ByteToInt : public FTypePromotion
-{
-	virtual bool ApplyPotentialPromotion(FEdGraphPinType& InOutType) OVERRIDE
+			const FVector2D ExitPos = GetNodePosition(GraphXBounds.Y + 1, 0);
+			ExitNode->NodePosX = ExitPos.X;
+			ExitNode->NodePosY = ExitPos.Y;
+		}
+				
+		bool bHasErrors = ((MessageLog.NumErrors - StartingErrorCount) > 0);
+		ActiveMessageLog = nullptr;
+		
+		return !bHasErrors;
+	}
+    
+    /**
+	 * Overloaded, part of the FExpressionVisitor interface; attempts to 
+     * generate either a vaiable-get node, an input pin, or a literal fragment
+     * from the supplied FTokenWrapperNode (all depends on the token's type).
+	 *
+	 * @return True to continue travesing the expression tree, false to stop.
+	 */
+	virtual bool Visit(FTokenWrapperNode& ExpressionNode, EVisitPhase Phase) OVERRIDE
 	{
-		const UEdGraphSchema_K2* Schema = GetDefault<UEdGraphSchema_K2>();
+		check(ActiveMessageLog != nullptr);
+		const UEdGraphSchema_K2* K2Schema = GetDefault<UEdGraphSchema_K2>();
 
-		if (InOutType.PinCategory == Schema->PC_Byte)
+		if (ExpressionNode.Token.TokenType == FBasicToken::TOKEN_Identifier)
 		{
-			InOutType.PinCategory = Schema->PC_Int;
-			InOutType.PinSubCategoryObject = NULL;
-			return true;
+			FString const VariableIdentifier = ExpressionNode.Token.Identifier;
+			// first we try to match up variables with existing variable properties on the blueprint
+			if (UProperty* VariableProperty = FindField<UProperty>(TargetBlueprint->SkeletonGeneratedClass, *VariableIdentifier))
+			{
+				TSharedPtr<FCodeGenFragment_VariableGet> VariableGetFragment = GeneratePropertyFragment(ExpressionNode, VariableProperty, *ActiveMessageLog);
+				if (VariableGetFragment.IsValid())
+				{
+					CompiledFragments.Add(&ExpressionNode, VariableGetFragment);
+				}
+			}
+			// if a variable-get couldn't be created for it, it needs to be an input to the math node
+			else 
+			{
+				CompiledFragments.Add(&ExpressionNode, GenerateInputPinFragment(VariableIdentifier));
+			}
+			
 		}
-
-		return false;
+		else if (ExpressionNode.Token.TokenType == FBasicToken::TOKEN_Const)
+		{
+			CompiledFragments.Add(&ExpressionNode, GenerateLiteralFragment(ExpressionNode.Token, *ActiveMessageLog));
+		}
+		else
+		{
+			ActiveMessageLog->Error(*LOCTEXT("UnexpectedTokenType", "Unhandled token type in expression: '@@'").ToString(), CompilingNode);
+		}
+		
+		// keep traversing the expression tree... we should handle cascading
+		// errors that result from ones incurred here, gathering them all as we
+		// go, presenting them to the user later
+		return true;
 	}
-};
 
-struct FTypePromotion_IntToFloat : public FTypePromotion
-{
-	virtual bool ApplyPotentialPromotion(FEdGraphPinType& InOutType) OVERRIDE
+	/**
+	 * Overloaded, part of the FExpressionVisitor interface... On VISIT_Post, 
+     * attempts to generate a UK2Node_CallFunction node for the specified
+     * FBinaryOperator.
+	 *
+	 * @return True to continue travesing the expression tree, false to stop.
+	 */
+	virtual bool Visit(FBinaryOperator& ExpressionNode, EVisitPhase Phase) OVERRIDE
 	{
-		const UEdGraphSchema_K2* Schema = GetDefault<UEdGraphSchema_K2>();
-
-		if (InOutType.PinCategory == Schema->PC_Int)
+		check(ActiveMessageLog != nullptr);
+		
+		// we only care about the "Post" visit, after the operands fragments have been generated
+		if (Phase == FExpressionVisitor::VISIT_Post)
 		{
-			InOutType.PinCategory = Schema->PC_Float;
-			InOutType.PinSubCategoryObject = NULL;
-			return true;
-		}
-
-		return false;
+			TSharedPtr<FCodeGenFragment> LHS = CompiledFragments.FindRef(&(ExpressionNode.LHS.Get()));
+			TSharedPtr<FCodeGenFragment> RHS = CompiledFragments.FindRef(&(ExpressionNode.RHS.Get()));
+			
+			TArray< TSharedPtr<FCodeGenFragment> > ArgumentList;
+			ArgumentList.Add(LHS);
+			ArgumentList.Add(RHS);
+			
+			TSharedPtr<FCodeGenFragment_FuntionCall> FunctionFragment = GenerateFunctionFragment(ExpressionNode, ExpressionNode.Operator, ArgumentList, *ActiveMessageLog);
+			if (FunctionFragment.IsValid())
+			{
+				CompiledFragments.Add(&ExpressionNode, FunctionFragment);
+			}
+        }
+        
+		// keep traversing the expression tree... we should handle cascading
+		// errors that result from ones incurred here, gathering them all as we
+		// go, presenting them to the user later
+		return true;
 	}
-};
-
-/////////////////////////////////////////////////////
-// FSuperDumbCodeGenerator
-
-// Exactly what it says on the tin
-class FSuperDumbCodeGenerator : public FExpressionVisitor
-{
-private:
-	// Outer node that contains the math expression being compiled
-	UK2Node_MathExpression* TopMathNode;
-
-	// List of known operators
-	TMap<FString, FOperatorTable> OperatorList;
 	
-	// List of allowable type promotions
-	FTypePromotion_IntToFloat PromoteIntToFloat;
-	FTypePromotion_ByteToInt PromoteByteToInt;
-	TArray<FTypePromotion*> TypePromotionOrdering;
-
-	bool bNoErrors;
-	UEdGraph* TargetGraph;
-	UBlueprint* TargetBP;
-
-	// Constants used in K2 node visual layout
-	float GenerationXOffset;
-	float PerGenerationOffset;
-
-	// Chart of the depth of the tree
-	FDepthLabelVisitor Depths;
-
-	// Chart of the height of the tree at a given depth
-	TMap<int32, int32> HeightAtEachDepthLevel;
-	TMap<FExpressionNode*, int32> HeightWithinLevelPerNode;
-
-	TMap<FExpressionNode*, TSharedPtr<FCodeGenFragment> > CompiledNodes;
-public:
-	FSuperDumbCodeGenerator(UK2Node_MathExpression* InNode)
+	/**
+	 *
+	 */
+	virtual bool Visit(FExpressionList& ExpressionNode, EVisitPhase Phase) OVERRIDE
 	{
-		TopMathNode = InNode;
-		TargetGraph = InNode->BoundGraph;
-		BuildOperatorTable();
-		bNoErrors = true;
-		GenerationXOffset = 0.0f;
-		PerGenerationOffset = 200.0f;
-		TargetBP = FBlueprintEditorUtils::FindBlueprintForGraphChecked(TargetGraph);
+		// no fragments are generated from a list node, it mostly acts as a link
+		// from a parent node to some set of sub-expressions
+		
+		
+		// keep traversing the expression tree... if there are any errors,
+		// they'll be caught in the children nodes (or maybe in the parent)
+		return true;
 	}
-
-	void GenerateCode(TSharedRef<FExpressionNode> Root)
+	
+	/**
+	 *
+	 */
+	virtual bool Visit(FFunctionExpression& ExpressionNode, EVisitPhase Phase) OVERRIDE
 	{
-		// Generate the depth chart (for pretty layout)
-		Root->Accept(Depths);
+		check(ActiveMessageLog != nullptr);
 		
-		// Generate the height chart
-		for (auto DepthIt = Depths.Chart.CreateConstIterator(); DepthIt; ++DepthIt)
-		{
-			FExpressionNode* Node = DepthIt.Key();
-			int32 NodeDepth = DepthIt.Value();
-
-			int32& HeightHere = HeightAtEachDepthLevel.FindOrAdd(NodeDepth);
-			HeightWithinLevelPerNode.Add(Node, HeightHere);
-			++HeightHere;
-		}
-
-		// Generate code
-		Root->Accept(*this);
-
-		// Connect the result
-		UK2Node_Tunnel* EntryNode = TopMathNode->GetEntryNode();
-		UK2Node_Tunnel* ExitNode = TopMathNode->GetExitNode();
-
-		TSharedPtr<FCodeGenFragment> RootFragment = GetFragmentForNode(Root);
-		if (RootFragment.IsValid())
+		// we only care about the "Post" visit, after the function's parameter fragments have been generated
+		if (Phase == FExpressionVisitor::VISIT_Post)
 		{
-			const FString ReturnPinName = TEXT("ReturnValue");
-			UEdGraphPin* ReturnPin = NULL;
-			if (UEdGraphPin* OldReturnPin = ExitNode->FindPin(ReturnPinName))
+			TArray< TSharedPtr<FCodeGenFragment> > ArgumentList;
+			for (TSharedRef<IFExpressionNode> Param : ExpressionNode.ParamList->Children)
 			{
-				if (OldReturnPin->PinType != RootFragment->Type)
-				{
-					OldReturnPin->BreakAllPinLinks();
-
-					// Remove all outputs (we'll create a result output when finished)
-					while (ExitNode->UserDefinedPins.Num() > 0)
-					{
-						ExitNode->RemoveUserDefinedPin(ExitNode->UserDefinedPins[0]);
-					}
-				}
-				else
-				{
-					ReturnPin = OldReturnPin;
-				}
+				TSharedPtr<FCodeGenFragment> ParamFragment = CompiledFragments.FindRef(&(Param.Get()));
+				ArgumentList.Add(ParamFragment);
 			}
 			
-			if (ReturnPin == NULL)
+			TSharedPtr<FCodeGenFragment_FuntionCall> FunctionFragment = GenerateFunctionFragment(ExpressionNode, ExpressionNode.FuncName, ArgumentList, *ActiveMessageLog);
+			if (FunctionFragment.IsValid())
 			{
-				ReturnPin = ExitNode->CreateUserDefinedPin(TEXT("ReturnValue"), RootFragment->Type);
+				CompiledFragments.Add(&ExpressionNode, FunctionFragment);
 			}
-			RootFragment->ConnectToInput(ReturnPin);
 		}
-
-		// Position the entry and exit nodes somewhere sane
+		
+		// keep traversing the expression tree... we should handle cascading
+		// errors that result from ones incurred here, gathering them all as we
+		// go, presenting them to the user later
+		return true;
+	}
+	
+	/**
+	 * From the FExpressionVisitor interface; where we would handle prefixed 
+	 * unary operators. Currently support for those is unimplemented, so we just
+	 * log a descriptive error and return.
+	 *
+	 * @return Always true, it is expected that cascading errors are handled (and all should be logged).
+	 */
+	virtual bool Visit(FUnaryOperator& ExpressionNode, EVisitPhase Phase) OVERRIDE
+	{
+		// don't want to double up on the error message (in the "Post" phase)
+		if (Phase == VISIT_Pre)
 		{
-			const FVector2D EntryPos = GetNodePosition(0, 0);
-			EntryNode->NodePosX = EntryPos.X;
-			EntryNode->NodePosY = EntryPos.Y;
-
-			const FVector2D ExitPos = GetNodePosition(Depths.GetMaximumDepth()+1, 0);
-			ExitNode->NodePosX = ExitPos.X;
-			ExitNode->NodePosY = ExitPos.Y;
+			FText ErrorText = FText::Format(LOCTEXT("UnaryExpressionError", "Currently, prefixed unary operators {0} are prohibited in expressions: '@@'"),
+				FText::FromString(ExpressionNode.ToString()));
+			
+			ActiveMessageLog->Error(*ErrorText.ToString(), CompilingNode);
 		}
+		
+		// keep traversing the expression tree... we should handle cascading
+		// errors that result from this, and gather them all to present to the user
+		return true;
 	}
-
-	TSharedPtr<FCodeGenFragment> GetFragmentForNode(TSharedPtr<FExpressionNode> Node) const
+	
+	/**
+	 * From the FExpressionVisitor interface; where we would handle conditional
+	 * ?: operators. Currently support for those is unimplemented, so we just
+	 * log a descriptive error and return.
+	 *
+	 * @return Always true, it is expected that cascading errors are handled (and all should be logged).
+	 */
+	virtual bool Visit(FConditionalOperator& ExpressionNode, EVisitPhase Phase) OVERRIDE
 	{
-		TSharedPtr<FCodeGenFragment> Result = CompiledNodes.FindRef(Node.Get());
-		return Result;
+		check(ActiveMessageLog != nullptr);
+		
+		// don't want to double up on the error message (in the "Post" phase)
+		if (Phase == VISIT_Pre)
+		{
+			FText ErrorText = FText::Format(LOCTEXT("ConditionalExpressionError", "Currently, conditional operators {0} are prohibited in expressions: '@@'"),
+				FText::FromString(ExpressionNode.ToString()));
+			
+			ActiveMessageLog->Error(*ErrorText.ToString(), CompilingNode);
+		}
+		
+		// keep traversing the expression tree... we should handle cascading
+		// errors that result from this, and gather them all to present to the user
+		return true;
 	}
-
-// 	virtual void Visit(class FExpressionNode& Node, EVisitPhase Phase) {}
-
-	void Error(FExpressionNode& Node, const FString& ErrorStr)
+	
+private:
+	/**
+	 * Another FExpressionVisitor interface function... A Generic catch all for 
+	 * any expression nodes that we don't explicitly handle. Simply logs an 
+	 * error, and returns.
+	 *
+	 * @return Always true, it is expected that cascading errors can be handled (and all should be logged).
+	 */
+	virtual bool VisitUnhandled(IFExpressionNode& ExpressionNode, EVisitPhase Phase) OVERRIDE
 	{
-		//@TODO: Make this stuff suck less
-		bNoErrors = false;
-		TopMathNode->ErrorMsg = ErrorStr;
-		TopMathNode->ErrorType = EMessageSeverity::Error;
-		TopMathNode->bHasCompilerMessage = true;
+		check(ActiveMessageLog != nullptr);
+		if (Phase == VISIT_Leaf || Phase == VISIT_Pre)
+		{
+			FText ErrorText = FText::Format(LOCTEXT("UnhandledExpressionNode", "Unsupported operation ({0}) in the expression: '@@'"),
+				FText::FromString(ExpressionNode.ToString()));
+		
+			ActiveMessageLog->Error(*ErrorText.ToString(), CompilingNode);
+		}
+		
+		// keep traversing the expression tree... we should handle cascading
+		// errors that result from this, and gather them all to present to the user
+		return true;
 	}
-
-	virtual void Visit(class FTokenWrapperNode& Node, EVisitPhase Phase)
-	{
-		const UEdGraphSchema_K2* K2Schema = GetDefault<UEdGraphSchema_K2>();
-
-		if (Node.Token.TokenType == TOKEN_Identifier)
+	
+	/**
+     * Either adds a new pin, or finds an existing one on the MathExpression
+     * node. From that, a fragment is generated (to track the pin, so
+     * connections can be made later).
+     *
+     * @param  VariableIdentifier   The name of the pin to generate a fragment for.
+     * @return A new input pin fragment (associated with a pin on the math expression's entry node).
+     */
+	TSharedPtr<FCodeGenFragment_InputPin> GenerateInputPinFragment(FString const VariableIdentifier)
+	{
+		TSharedPtr<FCodeGenFragment_InputPin> InputPinFragment;
+		UEdGraphSchema_K2 const* K2Schema = GetDefault<UEdGraphSchema_K2>();
+		
+		UK2Node_Tunnel* EntryNode = CompilingNode->GetEntryNode();
+		// if a pin under this name already exists, use that
+		if (UEdGraphPin* InputPin = EntryNode->FindPin(VariableIdentifier))
 		{
-			const FString Identifier = Node.Token.Identifier;
-
-			if (UProperty* VariableProperty = FindField<UProperty>(TargetBP->SkeletonGeneratedClass, *Identifier))
+			InputPinFragment = MakeShareable(new FCodeGenFragment_InputPin(InputPin));
+		}
+		// otherwise, a new input pin needs to be created for it
+		else
+		{
+			// Create an input pin (using the default guessed type)
+			FEdGraphPinType DefaultType;
+			// currently, generated expressions ALWAYS take a float (it is the most versatile type)
+			DefaultType.PinCategory = K2Schema->PC_Float;
+			
+			UEdGraphPin* NewInputPin = EntryNode->CreateUserDefinedPin(VariableIdentifier, DefaultType);
+			InputPinFragment = MakeShareable(new FCodeGenFragment_InputPin(NewInputPin));
+		}
+		
+		return InputPinFragment;
+	}
+	
+    /**
+     * Attempts to generate a VariableGet node for the blueprint graph. If one
+     * isn't generated, then this function logs an error (and returns an empty
+     * pointer). However, if one is successfully created, then a fragment
+     * wrapper is created and returned (to aid in linking the node later).
+     *
+     * @param  ExpressionContext    The expression node that we're generating this fragment for.
+     * @param  VariableProperty     The variable that the generated UK2Node should access.
+     * @return An empty pointer if we failed to generate something, otherwise new variable-get fragment.
+     */
+	TSharedPtr<FCodeGenFragment_VariableGet> GeneratePropertyFragment(FTokenWrapperNode& ExpressionContext, UProperty* VariableProperty, FCompilerResultsLog& MessageLog)
+	{
+		check(ExpressionContext.Token.TokenType == FBasicToken::TOKEN_Identifier);
+		check(VariableProperty != nullptr);
+		UEdGraphSchema_K2 const* K2Schema = GetDefault<UEdGraphSchema_K2>();
+		
+		TSharedPtr<FCodeGenFragment_VariableGet> VariableGetFragment;
+		
+		UClass* VariableAccessClass = TargetBlueprint->SkeletonGeneratedClass;
+		if (UEdGraphSchema_K2::CanUserKismetAccessVariable(VariableProperty, VariableAccessClass, UEdGraphSchema_K2::CannotBeDelegate))
+		{
+			FEdGraphPinType VarType;
+			if (K2Schema->ConvertPropertyToPinType(VariableProperty, /*out*/VarType))
 			{
-				UClass* VariableAccessClass = TargetBP->SkeletonGeneratedClass;
+				UK2Node_VariableGet* NodeTemplate = NewObject<UK2Node_VariableGet>();
+				NodeTemplate->VariableReference.SetSelfMember(VariableProperty->GetFName());
+				UK2Node_VariableGet* VariableGetNode = SpawnNodeFromTemplate<UK2Node_VariableGet>(&ExpressionContext, NodeTemplate);
 				
-				if (UEdGraphSchema_K2::CanUserKismetAccessVariable(VariableProperty, VariableAccessClass, UEdGraphSchema_K2::CannotBeDelegate))
-				{
-					FEdGraphPinType Type;
-					if (K2Schema->ConvertPropertyToPinType(VariableProperty, /*out*/ Type))
-					{
-						UK2Node_VariableGet* NodeTemplate = NewObject<UK2Node_VariableGet>();
-						NodeTemplate->VariableReference.SetSelfMember(VariableProperty->GetFName());
-						UK2Node_VariableGet* VariableGetNode = SpawnNodeFromTemplate<UK2Node_VariableGet>(&Node, NodeTemplate);
-
-						auto Fragment = MakeShareable(new FCodeGenFragment_VariableGet(VariableGetNode, Type));
-						CompiledNodes.Add(&Node, Fragment);
-					}
-					else
-					{
-						Error(Node, TEXT("Bad variable type"));
-					}
-				}
-				else
-				{
-					Error(Node, TEXT("Inaccessible variable"));
-				}
-			}
-			else if (UEdGraphPin* InputPin = TopMathNode->GetEntryNode()->FindPin(Identifier))
-			{
-				// It's an input pin
-				auto Fragment = MakeShareable(new FCodeGenFragment_InputPin(InputPin));
-				CompiledNodes.Add(&Node, Fragment);
+				VariableGetFragment = MakeShareable(new FCodeGenFragment_VariableGet(VariableGetNode, VarType));
 			}
 			else
 			{
-				// Create an input pin (using the default guessed type)
-				FEdGraphPinType DefaultType;
-				DefaultType.PinCategory = K2Schema->PC_Float;
-				
-				UK2Node_Tunnel* EntryNode = TopMathNode->GetEntryNode();
-				UEdGraphPin* NewInputPin = EntryNode->CreateUserDefinedPin(Identifier, DefaultType);
-
-				auto Fragment = MakeShareable(new FCodeGenFragment_InputPin(NewInputPin));
-				CompiledNodes.Add(&Node, Fragment);
+				FText ErrorText = FText::Format(LOCTEXT("IncompatibleVarError", "Blueprint '{0}' variable is incompatible with graph pins in the expression: '@@'"),
+					FText::FromName(VariableProperty->GetFName()));
+				MessageLog.Error(*ErrorText.ToString(), CompilingNode);
 			}
 		}
-		else if (Node.Token.TokenType == TOKEN_Const)
+		else
 		{
-			TSharedRef<FCodeGenFragment_Literal> Fragment = MakeShareable(new FCodeGenFragment_Literal);
-
-			Fragment->bSucceeded = true;
-			switch (Node.Token.ConstantType)
-			{
-			case CONST_Boolean:
-				Fragment->Type.PinCategory = K2Schema->PC_Boolean;
+			FText ErrorText = FText::Format(LOCTEXT("InaccessibleVarError", "Cannot access the blueprint's '{0}' variable from the expression: '@@'"),
+				FText::FromName(VariableProperty->GetFName()));
+			MessageLog.Error(*ErrorText.ToString(), CompilingNode);
+		}
+		
+		return VariableGetFragment;
+	}
+	
+    /**
+     * Spawns a fragment which wraps a literal value. No graph-node or pin is
+	 * created for this fragment; instead, it saves the literal value for later,
+	 * when this fragment is connected with another (it then enters the literal
+	 * value as the connecting pin's default).
+	 *
+	 * @param  ExpressionNode	The expression node that we're generating this fragment for.
+     * @return A new literal fragment.
+     */
+	TSharedPtr<FCodeGenFragment_Literal> GenerateLiteralFragment(FBasicToken const& Token, FCompilerResultsLog& MessageLog)
+	{
+		check(Token.TokenType == FBasicToken::TOKEN_Const);
+		UEdGraphSchema_K2 const* K2Schema = GetDefault<UEdGraphSchema_K2>();
+		
+		FEdGraphPinType LiteralType;
+		switch (Token.ConstantType)
+		{
+			case CPT_Bool:
+				LiteralType.PinCategory = K2Schema->PC_Boolean;
 				break;
-			case CONST_Float:
-				Fragment->Type.PinCategory = K2Schema->PC_Float;
+			case CPT_Float:
+				LiteralType.PinCategory = K2Schema->PC_Float;
 				break;
-			case CONST_Integer:
-				Fragment->Type.PinCategory = K2Schema->PC_Int;
+			case CPT_Int:
+				LiteralType.PinCategory = K2Schema->PC_Int;
 				break;
-			case CONST_String:
-				Fragment->Type.PinCategory = K2Schema->PC_String;
+			case CPT_String:
+				LiteralType.PinCategory = K2Schema->PC_String;
 				break;
 			default:
-				Error(Node, TEXT("Unknown literal type"));
-				Fragment->bSucceeded = false;
+				MessageLog.Error(*FText::Format(LOCTEXT("UnhandledLiteralType", "Unknown literal type in expression: '@@'"),
+					FText::AsNumber(Token.ConstantType)).ToString(),
+					CompilingNode);
 				break;
-			};
-			Fragment->DefaultValue = Node.Token.GetConstantValue();
-
-			CompiledNodes.Add(&Node, Fragment);
-		}
-		else
-		{
-			Error(Node, TEXT("Unexpected token type"));
-		}
+		};
+		
+		return MakeShareable(new FCodeGenFragment_Literal(Token.GetConstantValue(), LiteralType));
 	}
-
-	UFunction* FindMatchingFunction(FOperatorTable& Table, const TArray<FEdGraphPinType>& SourceTypeList)
-	{
-		TArray<FEdGraphPinType> TypeList = SourceTypeList;
-		/*
-			UFunction* FindFunction(const TArray<FEdGraphPinType>& DesiredTypes) const
+	
+	/**
+	 * Attempts to find a coresponding fuction (in this class's FOperatorTable),
+	 * one that matches the supplied operator name and the set of arguments. If
+	 * a matching function is found, then a wrapping UK2Node_CallFunction is
+	 * spawned and linked with the supplied arguments (otherwise, errors will
+	 * be logged and an empty pointer will be returned).
+	 *
+	 * @param  ExpressionContext	The expression node that we're generating this fragment for.
+	 * @param  FunctionName			The name of the operator (the key we're going to use looking up into FOperatorTable).
+	 * @param  ArgumentList			A set of other fragments that will plug in as parameters into function.
+	 * @return An empty pointer if we failed to generate something, otherwise the new function fragment.
+	 */
+	TSharedPtr<FCodeGenFragment_FuntionCall> GenerateFunctionFragment(IFExpressionNode& ExpressionContext, FString FunctionName, TArray< TSharedPtr<FCodeGenFragment> > ArgumentList, FCompilerResultsLog& MessageLog)
 	{
-		const UEdGraphSchema_K2* K2Schema = GetDefault<UEdGraphSchema_K2>();
-		if (Result == NULL)
-		{
-			// Try promoting
-		}
-	}
-*/
-		// Try to find the function
-		UFunction* Result = Table.FindFunction(TypeList);
-
-		// Try promoting if the initial pass failed
+		bool bMissingArgument = false;
 		
-		for (int32 PromotionPassIndex = 0; (PromotionPassIndex < TypePromotionOrdering.Num()) && (Result == NULL); ++PromotionPassIndex)
+		TArray<FEdGraphPinType> TypeList;
+		// create a type list from the argument fragments (so we can find a matching function signature)
+		for (int32 Index = 0; Index < ArgumentList.Num(); ++Index)
 		{
-			FTypePromotion& PromotionOperator = *TypePromotionOrdering[PromotionPassIndex];
-
-			// Apply the promotion operator to any values that match
-			bool bMadeChanges = false;
-			for (int32 Index = 0; Index < TypeList.Num(); ++Index)
-			{
-				bMadeChanges |= PromotionOperator.ApplyPotentialPromotion(TypeList[Index]);
-			}
-
-			// Try again
-			if (bMadeChanges)
+			if (!ArgumentList[Index].IsValid())
 			{
-				Result = Table.FindFunction(TypeList);
+				FText ErrorText = FText::Format(LOCTEXT("MissingArgument", "Failed to generate argument #{0} for the '{1}' function, in the expression: '@@'"),
+					FText::AsNumber(Index + 1),
+					FText::FromString(FunctionName));
+				
+				MessageLog.Error(*ErrorText.ToString(), CompilingNode);
+				
+				bMissingArgument = true;
+				continue;
 			}
+			TypeList.Add(ArgumentList[Index]->GetOutputType());
 		}
-
-		return Result;
-	}
-	
-	virtual void Visit(class FBinaryOperator& Node, EVisitPhase Phase)
-	{
-		if (Phase == VISIT_Post)
+		
+		TSharedPtr<FCodeGenFragment_FuntionCall> FunctionFragment;
+		if (!OperatorLookup.Contains(FunctionName))
+		{
+			FText ErrorText = FText::Format(LOCTEXT("UnknownFuncError", "Unknown function '{0}' in the expression: '@@'"), FText::FromString(FunctionName));
+			MessageLog.Error(*ErrorText.ToString(), CompilingNode);
+		}
+		else if (bMissingArgument)
 		{
-			if (OperatorList.Contains(Node.Operator))
+			// don't execute the other if-branches, head them off if there is already an error
+		}
+		else if (UFunction* MatchingFunction = OperatorLookup.FindMatchingFunction(FunctionName, TypeList))
+		{
+			UProperty* ReturnProperty = MatchingFunction->GetReturnProperty();
+			if (ReturnProperty == nullptr)
 			{
-				const UEdGraphSchema_K2* K2Schema = GetDefault<UEdGraphSchema_K2>();
-
-				FOperatorTable& Table = OperatorList.FindChecked(Node.Operator);
-
-				// Find the types of the left and right nodes
-				TSharedPtr<FCodeGenFragment> LHS = CompiledNodes.FindRef(&(Node.LHS.Get()));
-				TSharedPtr<FCodeGenFragment> RHS = CompiledNodes.FindRef(&(Node.RHS.Get()));
-				if (LHS.IsValid() && RHS.IsValid())
+				FText ErrorText = FText::Format(LOCTEXT("NoReturnTypeError", "The '{0}' function returns nothing, it cannot be used in the expression: '@@'"),
+					FText::FromString(FunctionName));
+				MessageLog.Error(*ErrorText.ToString(), CompilingNode);
+			}
+			else
+			{
+				UEdGraphSchema_K2 const* K2Schema = GetDefault<UEdGraphSchema_K2>();
+				
+				FEdGraphPinType ReturnType;
+				if (K2Schema->ConvertPropertyToPinType(ReturnProperty, /*out*/ReturnType))
 				{
-					TArray< TSharedPtr<FCodeGenFragment> > ArgumentList;
-					ArgumentList.Add(LHS);
-					ArgumentList.Add(RHS);
-
-					// Code below here is generalized to n-ary
-					TArray<FEdGraphPinType> TypeList;
-					for (int32 Index = 0; Index < ArgumentList.Num(); ++Index)
-					{
-						TypeList.Add(ArgumentList[Index]->Type);
-					}
-
-					UFunction* MatchingFunction = FindMatchingFunction(Table, TypeList);
-
-					// Try type promotions
-					if (MatchingFunction == NULL)
-					{
-						Error(Node, TEXT("No operator takes these types"));
-					}
-					else if (UProperty* ReturnProperty = MatchingFunction->GetReturnProperty())
+					UK2Node_CallFunction* NodeTemplate = NewObject<UK2Node_CallFunction>(CompilingNode->GetGraph());
+					NodeTemplate->SetFromFunction(MatchingFunction);
+					UK2Node_CallFunction* FunctionCall = SpawnNodeFromTemplate<UK2Node_CallFunction>(&ExpressionContext, NodeTemplate);
+					
+					int32 InitialErrorCount = MessageLog.NumErrors;
+					// connect this fragment to its children fragments
+					int32 PinWireIndex = 0;
+					for (auto PinIt = FunctionCall->Pins.CreateConstIterator(); PinIt; ++PinIt)
 					{
-						FEdGraphPinType ReturnType;
-						if (K2Schema->ConvertPropertyToPinType(ReturnProperty, /*out*/ ReturnType))
+						UEdGraphPin* InputPin = *PinIt;
+						if (!K2Schema->IsMetaPin(*InputPin) && (InputPin->Direction == EGPD_Input))
 						{
-							UK2Node_CallFunction* NodeTemplate = NewObject<UK2Node_CallFunction>(TopMathNode->GetGraph());
-							NodeTemplate->SetFromFunction(MatchingFunction);
-							UK2Node_CallFunction* FunctionCall = SpawnNodeFromTemplate<UK2Node_CallFunction>(&Node, NodeTemplate);
-
-							// Make connections
-							int32 PinWireIndex = 0;
-							for (auto PinIt = FunctionCall->Pins.CreateConstIterator(); PinIt; ++PinIt)
+							if (PinWireIndex < ArgumentList.Num())
 							{
-								UEdGraphPin* InputPin = *PinIt;
-								if (!K2Schema->IsMetaPin(*InputPin) && (InputPin->Direction == EGPD_Input))
+								TSharedPtr<FCodeGenFragment>& ArgumentNode = ArgumentList[PinWireIndex];
+								// try to make the connection (which might cause an error internally)
+								if (!ArgumentNode->ConnectToInput(InputPin, MessageLog))
 								{
-									if (PinWireIndex < ArgumentList.Num())
-									{
-										TSharedPtr<FCodeGenFragment>& ArgumentNode = ArgumentList[PinWireIndex];
-
-										// Try to make the connection (which might cause an error internally)
-										ArgumentNode->ConnectToInput(InputPin);
-											
-										if (!ArgumentNode->ErrorMessage.IsEmpty())
-										{
-											Error(Node, ArgumentNode->ErrorMessage);
-										}
-									}
-									else
-									{
-										// Too many pins - shouldn't be possible due to the checking in FindFunction above
-										Error(Node, TEXT("Too many pins"));
-									}
-									++PinWireIndex;
+									FText ErrorText = FText::Format(LOCTEXT("ConnectPinError", "Failed to connect parameter #{0} with input on '@@'"),
+										FText::AsNumber(PinWireIndex + 1));
+									MessageLog.Error(*ErrorText.ToString(), FunctionCall);
 								}
 							}
-
-							TSharedRef<FCodeGenFragment_FuntionCall> Fragment = MakeShareable(new FCodeGenFragment_FuntionCall(FunctionCall, ReturnType));
-							CompiledNodes.Add(&Node, Fragment);
-						}
-						else
-						{
-							Error(Node, TEXT("Bad pin type"));
+							else
+							{
+								// too many pins - shouldn't be possible due to the checking in FindMatchingFunction() above
+								FText ErrorText = LOCTEXT("ConnectPinError", "'@@' requires more parameters than were provided");
+								MessageLog.Error(*ErrorText.ToString(), FunctionCall);
+								break;
+							}
+							++PinWireIndex;
 						}
 					}
-					else
+					
+					bool bConnectionErrors = (InitialErrorCount < MessageLog.NumErrors);
+					if (bConnectionErrors)
 					{
-						Error(Node, TEXT("Bad function (no output value)"));
+						MessageLog.Error(*LOCTEXT("InternalExpressionError", "Internal node error for expression: '@@'").ToString(), CompilingNode);
 					}
+					
+					FunctionFragment = MakeShareable(new FCodeGenFragment_FuntionCall(FunctionCall, ReturnType));
 				}
 				else
 				{
-					Error(Node, TEXT("LHS or RHS had an error"));
+					FText ErrorText = FText::Format(LOCTEXT("ReturnTypeError", "The '{0}' function's return type is incompatible with graph pins in the expression: '@@'"),
+						FText::FromString(FunctionName));
+					MessageLog.Error(*ErrorText.ToString(), CompilingNode);
 				}
 			}
-			else
-			{
-				Error(Node, TEXT("Unknown operator"));
-			}
 		}
-	}
-// 	virtual void Visit(class FCastOperator& Node, EVisitPhase Phase) {}
-// 	virtual void Visit(class FUnaryOperator& Node, EVisitPhase Phase) {}
-// 	virtual void Visit(class FConditionalOperator& Node, EVisitPhase Phase) {}
-
-
-	void BuildOperatorTable()
-	{
-		// Map from function name to correct operator
-		TMap<FString, FString> OperatorFixups;
-		OperatorFixups.Add(TEXT("BooleanAND"), TEXT("&&"));
-		OperatorFixups.Add(TEXT("BooleanOR"), TEXT("||"));
-		OperatorFixups.Add(TEXT("BooleanXOR"), TEXT("^"));
-		OperatorFixups.Add(TEXT("Not_PreBool"), TEXT("!"));
-
-		// Run thru all functions and build up a list of ones from the function libraries that have good operator info
-		for (TObjectIterator<UClass> ClassIt; ClassIt; ++ClassIt)
+		else
 		{
-			UClass* TestClass = *ClassIt;
-			if (TestClass->IsChildOf(UBlueprintFunctionLibrary::StaticClass()) && (!TestClass->HasAnyClassFlags(CLASS_Abstract)))
-			{
-				for (TFieldIterator<UFunction> FuncIt(TestClass, EFieldIteratorFlags::ExcludeSuper); FuncIt; ++FuncIt)
-				{
-					UFunction* TestFunction = *FuncIt;
-					if (TestFunction->HasMetaData(FBlueprintMetadata::MD_CompactNodeTitle) && TestFunction->HasAnyFunctionFlags(FUNC_BlueprintPure) && (TestFunction->GetReturnProperty() != NULL))
-					{
-						const FString* pReplacementOperator = OperatorFixups.Find(TestFunction->GetName());
-						const FString Operator = (pReplacementOperator != NULL) ? (*pReplacementOperator) : TestFunction->GetMetaData(FBlueprintMetadata::MD_CompactNodeTitle);
-
-						FOperatorTable& Table = OperatorList.FindOrAdd(Operator);
-						Table.MatchingFunctions.Add(TestFunction);
-					}
-				}
-			}
+			FText ErrorText = FText::Format(LOCTEXT("OperatorParamsError", "Cannot find a '{0}' operator that takes the supplied param types, for expression: '@@'"),
+				FText::FromString(FunctionName));
+			MessageLog.Error(*ErrorText.ToString(), CompilingNode);
 		}
-
-		// Build the list of acceptable type promotions
-		TypePromotionOrdering.Add(&PromoteByteToInt);
-		TypePromotionOrdering.Add(&PromoteIntToFloat);
+		
+		return FunctionFragment;
 	}
-
-private:
+	
+	/**
+	 * Utility method to turn an FLayoutVisitor coordinate into graph coordinates.
+	 * FLayoutVisitor coordinates are in terms of nodes (so a depth of 1, would
+	 * mean one node to the right of the initial node).
+	 * 
+	 * @param  Depth	Horizontal coordinate (how many nodes away from the initial node).
+	 * @param  Height	Vertical coordinate (how many nodes down from the initial node).
+	 * @return A 2D graph position for the center of a node to be placed at.
+	 */
 	FVector2D GetNodePosition(int32 Depth, int32 Height) const
 	{
-		const float MiddleHeight = FMath::Max(HeightAtEachDepthLevel.FindRef(Depth), 1) * 0.5f;
+		// get a count of how many nodes there are at this specific depth
+		int32 TotalHeight = LayoutMapper.DepthHeightLookup.FindRef(LayoutMapper.GetMaximumDepth() - Depth);
+
+		const float MiddleHeight  = FMath::Max(TotalHeight, 1) * 0.5f;
 		const float HeightPerNode = 140.0f;
-		const float DepthPerNode = 240.0f;
+		const float DepthPerNode  = 240.0f;
 
 		return FVector2D(Depth * DepthPerNode, (Height - MiddleHeight + 0.5f) * HeightPerNode);
 	}
 
+	/**
+	 * Templatized function for turning an expression node into a UK2Node. This
+	 * takes the expression node's position in the expression tree and turns it 
+	 * into a blueprint graph position (placing the new UK2Node there).
+	 * 
+	 * @param  ForExpression	The expression node that the will UK2Node represent.
+	 * @param  Template			An instance of the UK2Node type you wish to spawn.
+	 * @return The newly created UK2Node.
+	 */
 	template<typename NodeType>
-	NodeType* SpawnNodeFromTemplate(FExpressionNode* ForExpression, NodeType* Template)
+	NodeType* SpawnNodeFromTemplate(IFExpressionNode* ForExpression, NodeType* Template)
 	{
-		const int32 Y = HeightWithinLevelPerNode.FindRef(ForExpression);
-		const int32 X = Depths.GetMaximumDepth() - Depths.Chart.FindRef(ForExpression);
+		const int32 Y = LayoutMapper.HeightChart.FindRef(ForExpression);
+		const int32 X = LayoutMapper.GetMaximumDepth() - LayoutMapper.DepthChart.FindRef(ForExpression);
+		
+		GraphXBounds.X = FMath::Min((int32)GraphXBounds.X, X);
+		GraphXBounds.Y = FMath::Max((int32)GraphXBounds.Y, X);
 
 		const FVector2D Location = GetNodePosition(X, Y);
-		return FEdGraphSchemaAction_K2NewNode::SpawnNodeFromTemplate<NodeType>(TargetGraph, Template, Location);
+		return FEdGraphSchemaAction_K2NewNode::SpawnNodeFromTemplate<NodeType>(CompilingNode->BoundGraph, Template, Location);
 	}
 
-};
-
-
-
-
-
-
-
-
-
-
+private:
+	/** The node that we're generating sub-nodes and pins for */
+	UK2Node_MathExpression* CompilingNode;
 
+	/** The blueprint that CompilingNode belongs to (the blueprint this will 
+	 * generate a graph for)
+	 */
+	UBlueprint* TargetBlueprint;
 
+	/** List of known operators, and mappings from them to associated functions */
+	FOperatorTable OperatorLookup;
 
+	/** 
+	 * An FLayoutVisitor that charts the depth of the expression tree (and what
+	 * depth/height each expression node is at). Used to layout the graph nicely.
+	 */
+	FLayoutVisitor LayoutMapper;
+	
+	/**
+	 * Supplements LayoutMapper, tracks where nodes were actually placed (sometimes
+	 * the depth of an expression node doesn't map one-to-one with the fragment
+	 * in the graph), so you have the min and max x locations of spawned graph nodes.
+	 */
+	FVector2D GraphXBounds;
 
+	/**
+	 * Fragments that represent spawned UK2Nodes or literals that were generated
+	 * from traversing the expression tree... These fragments facilitate
+	 * connections between each other (that's why we need to track them).
+	 */
+	TMap<IFExpressionNode*, TSharedPtr<FCodeGenFragment> > CompiledFragments;
+	
+	/** 
+	 * Used so the various Visit() methods have a way to log errors, null when 
+	 * not in the middle of GenerateCode().
+	 */
+	FCompilerResultsLog* ActiveMessageLog;
+};
 
-/////////////////////////////////////////////////////
+/*******************************************************************************
+ * FExpressionParser
+*******************************************************************************/
 
 #define PARSE_HELPER_BEGIN(NestedRuleName) \
-	TSharedRef<FExpressionNode> LHS = NestedRuleName(Context); \
+	TSharedRef<IFExpressionNode> LHS = NestedRuleName(); \
 	Begin:
 
 #define PARSE_HELPER_ENTRY(NestedRuleName, DesiredToken) \
-	if (MatchSymbol(DesiredToken)) \
+	if (IsValid() && MatchSymbol(DesiredToken)) \
 	{ \
-		TSharedRef<FExpressionNode> RHS = NestedRuleName(Context); \
+		TSharedRef<IFExpressionNode> RHS = NestedRuleName(); \
 		LHS = MakeShareable(new FBinaryOperator(DesiredToken, LHS, RHS)); \
 		goto Begin; \
 	} \
@@ -1892,23 +1735,79 @@ private:
 #define PARSE_HELPER_END \
 	{ return LHS; }
 
+/**
+ * Recursively builds an IFExpressionNode tree, where leaf nodes represent tokens 
+ * (constants, literals, or identifiers), and branch nodes represent operations
+ * on the attached children. The chaining order of expression functions is what 
+ * determines operator precedence.
+ */
+class FExpressionParser : public FBasicTokenParser
+{
+public:
+	/**
+	 * Takes a string and parses a mathematical expression out of it, returning 
+	 * the head of an expression tree that was generated as a result.
+	 * 
+	 * @param  InExpression		The string you wish to parse.
+	 * @return An expression node, which serves as the root of an expression tree representing the provided string.
+	 */
+	TSharedRef<IFExpressionNode> ParseExpression(FString InExpression)
+	{
+		ExpressionString = InExpression;
+		ResetParser(*ExpressionString);
+
+		return Expression();
+	}
 
+private:
+	/**
+	 * Starting point for parsing full expressions (sets off on parsing out 
+	 * operations according to operator precedence)... Could be used for the 
+	 * initial root expression, or various other sub-expressions (like those 
+	 * encapsulated in parentheses, etc.). 
+	 * 
+	 * @return Root node of an expression tree that was generated from where we 
+	 *         are in parsing ExpressionString (at time of calling).
+	 */
+	TSharedRef<IFExpressionNode> Expression()
+	{
+		// AssignmentExpression has the lowest precedence, start with it (it 
+		// will attempt to parse out higher precedent operations first)
+		return AssignmentExpression();
+	}
 
-class FExpressionParser : public FBaseParser
-{
-protected:
-	static TSharedRef<FExpressionNode> NullExpression;
+	/**
+	 * Intended to support assignment within the expression (setting temp or 
+	 * external variables equal to some value, so they can be used later in the 
+	 * expression).
+	 * 
+	 * @TODO   Implement!
+	 * @return Root node of an expression tree that was generated from where we 
+	 *         are in parsing ExpressionString (at time of calling).
+	 */
+	TSharedRef<IFExpressionNode> AssignmentExpression()
+	{
+		// ConditionalExpression takes precedence over a assignment operation, parse it first
+		return ConditionalExpression();
+	}
 
-protected:
-	TSharedRef<FExpressionNode> ConditionalExpression(FTreeContext& Context)
+	/**
+	 * Looks for a conditional ternary statement (c ? a : b) to parse, and 
+	 * tokenizes the operands. 
+	 * 
+	 * @return Root node of an expression tree that was generated from where we 
+	 *         are in parsing ExpressionString (at time of calling).
+	 */
+	TSharedRef<IFExpressionNode> ConditionalExpression()
 	{
-		TSharedRef<FExpressionNode> MainPart = LogicalOrExpression(Context);
+		// LogicalOrExpression takes precedence over a conditional operation, parse it first
+		TSharedRef<IFExpressionNode> MainPart = LogicalOrExpression();
 
-		if (MatchSymbol(TEXT("?")))
+		if (IsValid() && MatchSymbol(TEXT("?")))
 		{
-			TSharedRef<FExpressionNode> TruePart = Expression(Context);
+			TSharedRef<IFExpressionNode> TruePart = Expression();
 			RequireSymbol(TEXT(":"), TEXT("?: operator"));
-			TSharedRef<FExpressionNode> FalsePart = ConditionalExpression(Context);
+			TSharedRef<IFExpressionNode> FalsePart = ConditionalExpression();
 
 			return MakeShareable(new FConditionalOperator(MainPart, TruePart, FalsePart));
 		}
@@ -1918,51 +1817,107 @@ protected:
 		}
 	}
 
-	TSharedRef<FExpressionNode> LogicalOrExpression(FTreeContext& Context)
+	/**
+	 * Looks for a binary logical-or statement (a || b) to parse, and 
+	 * tokenizes the operands. 
+	 * 
+	 * @return Root node of an expression tree that was generated from where we 
+	 *         are in parsing ExpressionString (at time of calling).
+	 */
+	TSharedRef<IFExpressionNode> LogicalOrExpression()
 	{
-		PARSE_HELPER_BEGIN(LogicalAndExpression)
+		// LogicalOrExpression takes precedence over an or operation, parse it first
+		PARSE_HELPER_BEGIN(LogicalAndExpression) 
 		PARSE_HELPER_ENTRY(LogicalAndExpression, TEXT("||"))
 		PARSE_HELPER_END
 	}
 
-	TSharedRef<FExpressionNode> LogicalAndExpression(FTreeContext& Context)
+	/**
+	 * Looks for a binary logical-and statement (a && b) to parse, and 
+	 * tokenizes the operands. 
+	 * 
+	 * @return Root node of an expression tree that was generated from where we 
+	 *         are in parsing ExpressionString (at time of calling).
+	 */
+	TSharedRef<IFExpressionNode> LogicalAndExpression()
 	{
-		PARSE_HELPER_BEGIN(InclusiveOrExpression)
+		// InclusiveOrExpression takes precedence over an and operation, parse it first
+		PARSE_HELPER_BEGIN(InclusiveOrExpression) 
 		PARSE_HELPER_ENTRY(InclusiveOrExpression, TEXT("&&"))
 		PARSE_HELPER_END
 	}
 
-	TSharedRef<FExpressionNode> InclusiveOrExpression(FTreeContext& Context)
+	/**
+	 * Looks for a binary bitwise-or statement (a | b) to parse, and 
+	 * tokenizes the operands. 
+	 * 
+	 * @return Root node of an expression tree that was generated from where we 
+	 *         are in parsing ExpressionString (at time of calling).
+	 */
+	TSharedRef<IFExpressionNode> InclusiveOrExpression()
 	{
-		PARSE_HELPER_BEGIN(ExclusiveOrExpression)
+		// ExclusiveOrExpression takes precedence over an inclusive or operation, parse it first
+		PARSE_HELPER_BEGIN(ExclusiveOrExpression) 
 		PARSE_HELPER_ENTRY(ExclusiveOrExpression, TEXT("|"))
 		PARSE_HELPER_END
 	}
 
-	TSharedRef<FExpressionNode> ExclusiveOrExpression(FTreeContext& Context)
+	/**
+	 * Looks for a binary exclusive-or statement (a ^ b) to parse, and 
+	 * tokenizes the operands. 
+	 * 
+	 * @return Root node of an expression tree that was generated from where we 
+	 *         are in parsing ExpressionString (at time of calling).
+	 */
+	TSharedRef<IFExpressionNode> ExclusiveOrExpression()
 	{
-		PARSE_HELPER_BEGIN(AndExpression)
+		// AndExpression takes precedence over an exclusive or operation, parse it first
+		PARSE_HELPER_BEGIN(AndExpression) 
 		PARSE_HELPER_ENTRY(AndExpression, TEXT("^"))
 		PARSE_HELPER_END
 	}
 
-	TSharedRef<FExpressionNode> AndExpression(FTreeContext& Context)
+	/**
+	 * Looks for a binary bitwise-and statement (a & b) to parse, and 
+	 * tokenizes the operands. 
+	 * 
+	 * @return Root node of an expression tree that was generated from where we 
+	 *         are in parsing ExpressionString (at time of calling).
+	 */
+	TSharedRef<IFExpressionNode> AndExpression()
 	{
+		// EqualityExpression takes precedence over an and operation, parse it first
 		PARSE_HELPER_BEGIN(EqualityExpression)
 		PARSE_HELPER_ENTRY(EqualityExpression, TEXT("&"))
 		PARSE_HELPER_END
 	}
 
-	TSharedRef<FExpressionNode> EqualityExpression(FTreeContext& Context)
+	/**
+	 * Looks for a binary equality statement (like [a == b], or [a != b]) to 
+	 * parse, and tokenizes the operands. 
+	 * 
+	 * @return Root node of an expression tree that was generated from where we 
+	 *         are in parsing ExpressionString (at time of calling).
+	 */
+	TSharedRef<IFExpressionNode> EqualityExpression()
 	{
+		// RelationalExpression takes precedence over an equality expression, parse it first
 		PARSE_HELPER_BEGIN(RelationalExpression)
 		PARSE_HELPER_ENTRY(RelationalExpression, TEXT("=="))
 		PARSE_HELPER_ENTRY(RelationalExpression, TEXT("!="))
 		PARSE_HELPER_END
 	}
 
-	TSharedRef<FExpressionNode> RelationalExpression(FTreeContext& Context)
+	/**
+	 * Looks for a binary comparison statement to parse (like [a > b], [a <= b], 
+	 * etc.), and tokenizes the operands. 
+	 * 
+	 * @return Root node of an expression tree that was generated from where we 
+	 *         are in parsing ExpressionString (at time of calling).
+	 */
+	TSharedRef<IFExpressionNode> RelationalExpression()
 	{
+		// ShiftExpression takes precedence over a relational expression, parse it first
 		PARSE_HELPER_BEGIN(ShiftExpression)
 		PARSE_HELPER_ENTRY(ShiftExpression, TEXT("<"))
 		PARSE_HELPER_ENTRY(ShiftExpression, TEXT(">"))
@@ -1971,24 +1926,48 @@ protected:
 		PARSE_HELPER_END
 	}
 
-	TSharedRef<FExpressionNode> ShiftExpression(FTreeContext& Context)
+	/**
+	 * Looks for a binary bitwise shift statement to parse (like [a << b], or 
+	 * [a >> b]), and tokenizes the operands. 
+	 * 
+	 * @return Root node of an expression tree that was generated from where we 
+	 *         are in parsing ExpressionString (at time of calling).
+	 */
+	TSharedRef<IFExpressionNode> ShiftExpression()
 	{
+		// AdditiveExpression takes precedence over a shift, parse it first
 		PARSE_HELPER_BEGIN(AdditiveExpression)
 		PARSE_HELPER_ENTRY(AdditiveExpression, TEXT("<<"))
 		PARSE_HELPER_ENTRY(AdditiveExpression, TEXT(">>"))
 		PARSE_HELPER_END
 	}
 
-	TSharedRef<FExpressionNode> AdditiveExpression(FTreeContext& Context)
+	/**
+	 * Looks for a binary addition/subtraction statement to parse ([a + b], or 
+	 * [a - b]), and tokenizes the operands. 
+	 * 
+	 * @return Root node of an expression tree that was generated from where we 
+	 *         are in parsing ExpressionString (at time of calling).
+	 */
+	TSharedRef<IFExpressionNode> AdditiveExpression()
 	{
+		// MultiplicativeExpression takes precedence over an add/subtract, parse it first
 		PARSE_HELPER_BEGIN(MultiplicativeExpression)
 		PARSE_HELPER_ENTRY(MultiplicativeExpression, TEXT("+"))
 		PARSE_HELPER_ENTRY(MultiplicativeExpression, TEXT("-"))
 		PARSE_HELPER_END
 	}
 
-	TSharedRef<FExpressionNode> MultiplicativeExpression(FTreeContext& Context)
+	/**
+	 * Looks for a binary multiplication/division/modulus statement to parse 
+	 * ([a * b], [a / b], or [a % b]), and tokenizes the operands. 
+	 * 
+	 * @return Root node of an expression tree that was generated from where we 
+	 *         are in parsing ExpressionString (at time of calling).
+	 */
+	TSharedRef<IFExpressionNode> MultiplicativeExpression()
 	{
+		// CastExpression takes precedence over a multiply/division/modulus, parse it first
 		PARSE_HELPER_BEGIN(CastExpression)
 		PARSE_HELPER_ENTRY(CastExpression, TEXT("*"))
 		PARSE_HELPER_ENTRY(CastExpression, TEXT("/"))
@@ -1996,246 +1975,228 @@ protected:
 		PARSE_HELPER_END
 	}
 
-	TSharedRef<FExpressionNode> CastExpression(FTreeContext& Context)
-	{
-		//@TODO: support casts (currently this is too greedy, and messes up "4*(5)" interpreting (5) as a cast)
-// 		if (MatchSymbol(TEXT("(")))
-// 		{
-// 			//@TODO: Need to support qualifiers / typedefs / what have you
-// 			FToken TypeName;
-// 			GetToken(TypeName);
-// 			TSharedRef<FExpressionNode> TypeExpression = MakeShareable(new FTokenWrapperNode(TypeName));
-// 
-// 			RequireSymbol(TEXT(")"), TEXT("Closing ) in cast"));
-// 
-// 			TSharedRef<FExpressionNode> ValueExpression = CastExpression(Context);
-// 			
-// 			return MakeShareable(new FCastOperator(TypeExpression, ValueExpression));
-// 		}
-// 		else
+	/**
+	 * Intended to handle type-casts (like from float to int, etc.).
+	 * 
+	 * @TODO   Implement!
+	 * @return Root node of an expression tree that was generated from where we 
+	 *         are in parsing ExpressionString (at time of calling).
+	 */
+	TSharedRef<IFExpressionNode> CastExpression()
+	{
+		// @TODO: support casts (currently this is too greedy, and messes up "4*(5)" interpreting (5) as a cast)
+		// if (MatchSymbol(TEXT("(")))
+		// {
+		// 	//@TODO: Need to support qualifiers / typedefs / what have you
+		// 	FBasicToken TypeName;
+		// 	GetToken(TypeName);
+		// 	TSharedRef<IFExpressionNode> TypeExpression = MakeShareable(new FTokenWrapperNode(TypeName));
+		// 
+		// 	RequireSymbol(TEXT(")"), TEXT("Closing ) in cast"));
+		// 
+		// 	TSharedRef<IFExpressionNode> ValueExpression = CastExpression(Context);
+		// 			
+		// 	return MakeShareable(new FCastOperator(TypeExpression, ValueExpression));
+		// }
+		// else
 		{
-			return UnaryExpression(Context);
+			return UnaryExpression();
 		}
 	}
 
-	TSharedRef<FExpressionNode> UnaryExpression(FTreeContext& Context)
-	{
-// 		<unary-expression> ::= <postfix-expression>
-// 			| ++ <unary-expression>
-// 			| -- <unary-expression>
-// 			| <unary-operator> <cast-expression>
-// 			| sizeof <unary-expression>
-// 			| sizeof <type-name>
-		// 		<unary-operator> ::= &
-		// 			| *
-		// 			| +
-		// 			| -
-		// 			| ~
-		// 			| !
+	/**
+	 * Attempts to parse various unary statements (like positive/negative 
+	 * markers, logical negation, pre increment/decrement, etc.)
+	 * 
+	 * @return Root node of an expression tree that was generated from where we 
+	 *         are in parsing ExpressionString (at time of calling).
+	 */
+	TSharedRef<IFExpressionNode> UnaryExpression()
+	{
+		// 	prefix increment:    ++<unary-expression> 
+		// 	prefix decrement:    --<unary-expression> 
+		// 	bitwise not:          ~<unary-expression> 
+		// 	logical not:          !<unary-expression> 
+		// 	positive sign:        +<unary-expression> 
+		// 	negative sign:        -<unary-expression> 
+		// 	reference:            &<unary-expression> 
+		// 	dereference:          *<unary-expression> 
+		// 	negative sign:        -<unary-expression> 
+		// 	allocation:        new <unary-expression> 
+		// 	deallocation:   delete <unary-expression> 
+		// 	parameter pack: sizeof <unary-expression> 
+		// 	C-style cast:   (type) <unary-expression> 
+
+		//
+		// check for the various prefix operators and jump back to 
+		// CastExpression() for parsing the right operand...
 
 		if (MatchSymbol(TEXT("&")))
 		{
-			return MakeShareable(new FUnaryOperator(TEXT("&"), CastExpression(Context)));
+			return MakeShareable(new FUnaryOperator(TEXT("&"), CastExpression()));
 		}
 		else if (MatchSymbol(TEXT("+")))
 		{
-			return MakeShareable(new FUnaryOperator(TEXT("+"), CastExpression(Context)));
+			// would return pre-increment operators like so: 
+			//		unaryOp(+).RHS = unaryOp(+)
+			return MakeShareable(new FUnaryOperator(TEXT("+"), CastExpression()));
 		}
 		else if (MatchSymbol(TEXT("-")))
 		{
-			return MakeShareable(new FUnaryOperator(TEXT("-"), CastExpression(Context)));
+			// would return post-increment operators like so: 
+			//		unaryOp(-).RHS = unaryOp(-)
+			return MakeShareable(new FUnaryOperator(TEXT("-"), CastExpression()));
 		}
 		else if (MatchSymbol(TEXT("~")))
 		{
-			return MakeShareable(new FUnaryOperator(TEXT("~"), CastExpression(Context)));
+			return MakeShareable(new FUnaryOperator(TEXT("~"), CastExpression()));
 		}
 		else if (MatchSymbol(TEXT("!")))
 		{
-			return MakeShareable(new FUnaryOperator(TEXT("!"), CastExpression(Context)));
+			return MakeShareable(new FUnaryOperator(TEXT("!"), CastExpression()));
 		}
 		else
 		{
-			return PostfixExpression(Context);
+			return PostfixExpression();
 		}
 	}
 
-	TSharedRef<FExpressionNode> PostfixExpression(FTreeContext& Context)
-	{
-		//@TODO: Giant hack to get something working
-// 		if (MatchSymbol(TEXT("[")))
-// 		{
-// 			// Array indexing
-// 			TSharedRef<FExpressionNode> IndexExpression = Expression(Context);
-// 			RequireSymbol(TEXT("]"), TEXT("Closing ] in array indexing"));
-// 
-// 			return NullExpression;//@TODO
-// 		}
-// 		else if (MatchSymbol(TEXT("(")))
-// 		{
-// 			while (!PeekSymbol(TEXT(")")))
-// 			{
-// 				TSharedRef<FExpressionNode> Item = AssignmentExpression(Context);
-// //@TODO:
-// 			}
-// 
-// 			RequireSymbol(TEXT(")"), TEXT("Closing ) in function call"));
-// 
-// 			return NullExpression;//@TODO
-// 		}
-// 		else if (MatchSymbol(TEXT(".")) || MatchSymbol(TEXT("->")))
-// 		{
-// 			// Member reference
-// 			FToken Identifier;
-// 			if (GetIdentifier(Identifier, /*bNoConsts=*/ true))
-// 			{
-// 				//@TODO: Do stuffs
-// 			}
-// 			else
-// 			{
-// 				//@TODO: error
-// 			}
-// 
-// 			return NullExpression;//@TODO
-// 		}
-// 		else
-		{
-			return PrimaryExpression(Context);
+	/**
+	 * Intended to handle postfix operations (like post increment/decrement, 
+	 * array subscripting, member access, etc.).
+	 * 
+	 * @TODO   Implement!
+	 * @return Root node of an expression tree that was generated from where we 
+	 *         are in parsing ExpressionString (at time of calling).
+	 */
+	TSharedRef<IFExpressionNode> PostfixExpression()
+	{
+		// if (MatchSymbol(TEXT("[")))
+		// {
+		// 	// Array indexing
+		// 	TSharedRef<IFExpressionNode> IndexExpression = Expression(Context);
+		// 	RequireSymbol(TEXT("]"), TEXT("Closing ] in array indexing"));
+		// 
+		// 	return NullExpression; // @TODO: return a valid expression node
+		// }
+		// else if (MatchSymbol(TEXT("(")))
+		// {
+		// 	while (!PeekSymbol(TEXT(")")))
+		// 	{
+		// 		TSharedRef<IFExpressionNode> Item = AssignmentExpression(Context);
+		// 	}
+		// 
+		// 	RequireSymbol(TEXT(")"), TEXT("Closing ) in function call"));
+		// 
+		// 	return NullExpression; // @TODO: return a valid expression node
+		// }
+		// else if (MatchSymbol(TEXT(".")) || MatchSymbol(TEXT("->")))
+		// {
+		// 	// Member reference
+		// 	FBasicToken Identifier;
+		// 	if (GetIdentifier(Identifier, /*bNoConsts=*/ true))
+		// 	{
+		// 		//@TODO: Do stuffs
+		// 	}
+		// 	else
+		// 	{
+		// 		// @TODO: error
+		// 	}
+		// 
+		// 	return NullExpression; // @TODO: return a valid expression node
+		// }
+		// else
+		{
+			return PrimaryExpression();
 		}
 	}
 
-	TSharedRef<FExpressionNode> PrimaryExpression(FTreeContext& Context)
+	/**
+	 * End of the line, where we attempt to generate a leaf node (an identifier,
+	 * const literal, or a string). However, here we also look for the start of 
+	 * a sub-expression (one encapsulated in parentheses).
+	 * 
+	 * @return Either a leaf node (representing a variable, or literal), or another 
+	 *         branch node, representing a sub-expression encapsulated by parentheses.
+	 */
+	TSharedRef<IFExpressionNode> PrimaryExpression()
 	{
 		if (MatchSymbol(TEXT("(")))
 		{
-			TSharedRef<FExpressionNode> Result = Expression(Context);
-			RequireSymbol(TEXT(")"), TEXT("Closing ) in grouping"));
+			TSharedRef<IFExpressionNode> Result = Expression();
+			RequireSymbol(TEXT(")"), TEXT("group closing"));
 			return Result;
 		}
 		else
 		{
-			// Identifier, constant, or string
-			FToken Token;
+			// identifier, constant, or string
+			FBasicToken Token;
 			GetToken(Token);
-
-			return MakeShareable(new FTokenWrapperNode(Token));
-		}
-	}
-
-	TSharedRef<FExpressionNode> Expression(FTreeContext& Context)
-	{
-		return AssignmentExpression(Context);
-
-		//@TODO: Hackery again
-// 		TSharedRef<FExpressionList> ListNode = MakeShareable(new FExpressionList);
-// 		
-// 		ListNode->Add(AssignmentExpression(Context));
-// 
-// 		while (MatchSymbol(TEXT(",")))
-// 		{
-// 			ListNode->Add(AssignmentExpression(Context));
-// 		}
-// 
-// 		return ListNode;
-	}
-
-	TSharedRef<FExpressionNode> AssignmentExpression(FTreeContext& Context)
-	{
-		// No assignment expressions here...
-		return ConditionalExpression(Context);
-	}
-
-
-public:
-	TSharedRef<FExpressionNode> ParseExpression(FString InExpression)
-	{
-		ExpressionString = InExpression;
-		ResetParser(*ExpressionString, 1);
-
-		FTreeContext Dummy;
-		return Expression(Dummy);
-	}
-
-protected:
-	FString ExpressionString;
-};
-
-
-TSharedRef<FExpressionNode> FExpressionParser::NullExpression = MakeShareable(new FExpressionNode());
-
-
-class FExpressionRefreshHelper
-{
-protected:
-	UK2Node_MathExpression* Node;
-	FString DiagnosticString;
-	TSharedPtr<FExpressionNode> ParsedExpression;
-public:
-	FExpressionRefreshHelper(UK2Node_MathExpression* InNode, const FString& InExpression)
-		: Node(InNode)
-	{
-		Node->ErrorMsg.Empty();
-		Node->bHasCompilerMessage = false;
-
-		UBlueprint* BP = FBlueprintEditorUtils::FindBlueprintForNodeChecked(Node);
-
-		FExpressionParser Parser;
-		ParsedExpression = Parser.ParseExpression(InExpression);
-
-		DiagnosticString = ParsedExpression->ToString();
-
-		// Clear out old nodes
-		DeleteGeneratedNodesInGraph(Node->BoundGraph, false);
-
-		// Generate nodes
-		FSuperDumbCodeGenerator Generator(Node);
-		Generator.GenerateCode(ParsedExpression.ToSharedRef());
-
-		// Refresh the node since the connections may have changed
-		Node->ReconstructNode();
-
-		// Finally, recompile
-		FBlueprintEditorUtils::MarkBlueprintAsStructurallyModified(BP);
-	}
-
-	FString GetDiagnosticString() const
-	{
-		return DiagnosticString;
-	}
-
-protected:
-
-	static void DeleteGeneratedNodesInGraph(UEdGraph* Graph, bool bMarkAsModified = true)
-	{
-		UBlueprint* BP = FBlueprintEditorUtils::FindBlueprintForGraphChecked(Graph);
-		for (int32 NodeIndex = 0; NodeIndex < Graph->Nodes.Num(); )
-		{
-			UEdGraphNode* Node = Graph->Nodes[NodeIndex];
-			if (ExactCast<UK2Node_Tunnel>(Node) != NULL)
+			
+			// or maybe a function call?
+			if (MatchSymbol(TEXT("(")))
 			{
-				++NodeIndex;
+				TSharedPtr<FExpressionList> FuncArguments;
+				// if this is an empty function (takes no parameters)
+				if (PeekSymbol(TEXT(")")))
+				{
+					FuncArguments = MakeShareable(new FExpressionList);
+				}
+				else
+				{
+					FuncArguments = ListExpression();
+				}
+				
+				TSharedRef<IFExpressionNode> FuncExpression = MakeShareable(new FFunctionExpression(Token.Identifier, FuncArguments.ToSharedRef()));
+				
+				FText RequireError = FText::Format(LOCTEXT("MissingFuncClose", "'{0}' closing"), FText::FromString(Token.Identifier));
+				RequireSymbol(TEXT(")"), *RequireError.ToString());
+				
+				return FuncExpression;
 			}
 			else
 			{
-				FBlueprintEditorUtils::RemoveNode(BP, Node, true);
+				return MakeShareable(new FTokenWrapperNode(Token));
 			}
 		}
-		
-		if (bMarkAsModified)
+	}
+    
+	/**
+	 * Parses out a comma separated list of sub-expressions (arguments for a 
+	 * function or struct).
+	 *
+	 * @return A branch FExpressionList node, which holds a series of sub-expressions.
+	 */
+    TSharedRef<FExpressionList> ListExpression()
+	{
+		TSharedRef<FExpressionList> ListNode = MakeShareable(new FExpressionList);	
+		do 
 		{
-			FBlueprintEditorUtils::MarkBlueprintAsStructurallyModified(BP);
-		}
+			ListNode->Children.Add(Expression());
+		} while (MatchSymbol(TEXT(",")));
+
+		return ListNode;
 	}
+
+private:
+	/** The intact expression string that this is currently in charge of parsing */
+	FString ExpressionString;
 };
 
-/////////////////////////////////////////////////////
-// UK2Node_MathExpression
+/*******************************************************************************
+ * UK2Node_MathExpression
+*******************************************************************************/
 
+//------------------------------------------------------------------------------
 UK2Node_MathExpression::UK2Node_MathExpression(const FPostConstructInitializeProperties& PCIP)
 	: Super(PCIP)
 {
+	// renaming the node rebuilds the expression (the node name is where they 
+	// specify the math equation)
 	bCanRenameNode = true;
 }
 
+//------------------------------------------------------------------------------
 void UK2Node_MathExpression::PostEditChangeProperty(struct FPropertyChangedEvent& PropertyChangedEvent)
 {
 	Super::PostEditChangeProperty(PropertyChangedEvent);
@@ -2243,70 +2204,165 @@ void UK2Node_MathExpression::PostEditChangeProperty(struct FPropertyChangedEvent
 	const FName PropertyName = (PropertyChangedEvent.Property != NULL) ? PropertyChangedEvent.Property->GetFName() : NAME_None;
 	if (PropertyName == GET_MEMBER_NAME_CHECKED(UK2Node_MathExpression, Expression))
 	{
-		RebuildExpression();
+		RebuildExpression(Expression);
 	}
 }
 
+//------------------------------------------------------------------------------
 void UK2Node_MathExpression::GetMenuEntries(FGraphContextMenuBuilder& ContextMenuBuilder) const
 {
 	UK2Node_MathExpression* TemplateNode = NewObject<UK2Node_MathExpression>(GetTransientPackage(), GetClass());
 
 	const FString Category = TEXT("");
-	const FText MenuDesc = FText::FromString(TEXT("Add Math Expression..."));
-	const FString Tooltip = TEXT("Create a new mathematical expression");
+	const FText   MenuDesc = LOCTEXT("AddMathExprMenuOption", "Add Math Expression...");
+	const FString Tooltip  = TEXT("Create a new mathematical expression");
 
 	TSharedPtr<FEdGraphSchemaAction_K2NewNode> NodeAction = FK2ActionMenuBuilder::AddNewNodeAction(ContextMenuBuilder, Category, MenuDesc, Tooltip);
 	NodeAction->NodeTemplate = TemplateNode;
-	NodeAction->SearchTitle = TemplateNode->GetNodeSearchTitle();
+	NodeAction->SearchTitle  = TemplateNode->GetNodeSearchTitle();
 }
 
+//------------------------------------------------------------------------------
 TSharedPtr<class INameValidatorInterface> UK2Node_MathExpression::MakeNameValidator() const
 {
-	return MakeShareable( new FDummyNameValidator(EValidatorResult::Ok) );
+	// we'll let our parser mark the node for errors after the face (once the 
+	// name is submitted)... parsing it with every character could be slow
+	return MakeShareable(new FDummyNameValidator(EValidatorResult::Ok));
 }
 
+//------------------------------------------------------------------------------
 void UK2Node_MathExpression::OnRenameNode(const FString& NewName)
 {
-	Expression = NewName;
-	RebuildExpression();
+	RebuildExpression(NewName);
 }
 
-void UK2Node_MathExpression::RebuildExpression()
+//------------------------------------------------------------------------------
+void UK2Node_MathExpression::RebuildExpression(FString InExpression)
 {
-	if (GMessingWithMathNode)
+	static bool bIsAlreadyRebuilding = false;
+	// the rebuild can invoke a ReconstructNode(), which triggers this again,
+	// so this combined with the following 
+	if (!bIsAlreadyRebuilding)
 	{
-		return;
+		TGuardValue<bool> RecursionGuard(bIsAlreadyRebuilding, true);
+
+		ClearExpression();
+		Expression = InExpression;
+
+		if (!InExpression.IsEmpty()) // @TODO: is this needed?
+		{
+			// build a expression tree from the string
+			FExpressionParser Parser;
+			TSharedPtr<IFExpressionNode> ExpressionRoot = Parser.ParseExpression(InExpression);
+
+			// if the parser successfully chewed through the string
+			if (Parser.IsValid())
+			{
+				FMathGraphGenerator GraphGenerator(this);
+				// generate new nodes from the expression tree (could result in
+				// a series of errors being attached to the node).
+				if (!GraphGenerator.GenerateCode(ExpressionRoot.ToSharedRef(), *CachedMessageLog))
+				{
+					CachedMessageLog->Error(*LOCTEXT("MathExprFailedGen", "Failed to generate full expression graph for: '@@'").ToString(), this);
+				}
+			}
+			else
+			{
+				FText ErrorText = FText::Format(LOCTEXT("MathExprParseError", "PARSE ERROR in '@@': {0}"),
+					Parser.GetErrorState().Description);
+				CachedMessageLog->Error(*ErrorText.ToString(), this);
+			}
+		}
+
+		// refresh the node since the connections may have changed
+		Super::ReconstructNode();
+
+		// finally, recompile
+		UBlueprint* Blueprint = FBlueprintEditorUtils::FindBlueprintForNodeChecked(this);
+		FBlueprintEditorUtils::MarkBlueprintAsStructurallyModified(Blueprint);
 	}
+}
 
-	TGuardValue<bool> RecursionHelper(GMessingWithMathNode, true);
+//------------------------------------------------------------------------------
+void UK2Node_MathExpression::ClearExpression()
+{
+	// clear any errors 
+	SetNodeError(this, FText::GetEmpty());
+
+	// clear out old nodes
+	DeleteGeneratedNodesInGraph(BoundGraph);
 
-	// Rebuild
-	if (!Expression.IsEmpty())//@TODO: not needed?
+	// delete old pins
+	if (UK2Node_Tunnel* EntryNode = GetEntryNode())
+	{
+		// iterate back wards so we can remove as we go
+		for (int32 pinIndex = EntryNode->UserDefinedPins.Num()-1; pinIndex >= 0; --pinIndex)
+		{
+			TSharedPtr<FUserPinInfo> PinInfo = EntryNode->UserDefinedPins[pinIndex];
+			EntryNode->RemoveUserDefinedPin(PinInfo);
+		}
+	}
+	if (UK2Node_Tunnel* ExitNode = GetExitNode())
 	{
-		FExpressionRefreshHelper ParseHelper(this, Expression);
+		// iterate back wards so we can remove as we go
+		for (int32 pinIndex = ExitNode->UserDefinedPins.Num()-1; pinIndex >= 0; --pinIndex)
+		{
+			TSharedPtr<FUserPinInfo> PinInfo = ExitNode->UserDefinedPins[pinIndex];
+			ExitNode->RemoveUserDefinedPin(PinInfo);
+		}
+	}
+	
+	CachedMessageLog = MakeShareable(new FCompilerResultsLog);
+	
+	Expression.Empty();
+}
 
-		//@TODO: Debug code
-		ParseResults = ParseHelper.GetDiagnosticString();
+//------------------------------------------------------------------------------
+void UK2Node_MathExpression::ValidateNodeDuringCompilation(FCompilerResultsLog& MessageLog) const
+{
+	if (CachedMessageLog.IsValid())
+	{
+		MessageLog.Append(*CachedMessageLog);
+	}
+	// else, this may be some intermediate node in the compile, let's look at the errors from the original...
+	else if (UK2Node_MathExpression const* SourceNode = MessageLog.FindSourceObjectTypeChecked<UK2Node_MathExpression>(this))
+	{
+		// if the expressions match, then the errors should
+		check(SourceNode->Expression == Expression);
+		
+		// take the same errors from the original node (so we don't have to
+		// re-parse/re-gen to fish out the same errors)
+		if (SourceNode->CachedMessageLog.IsValid())
+		{
+			MessageLog.Append(*SourceNode->CachedMessageLog);
+		}
 	}
 }
 
+//------------------------------------------------------------------------------
 FText UK2Node_MathExpression::GetNodeTitle(ENodeTitleType::Type TitleType) const
 {
 	FString Result = Expression;
-
-	if (TitleType == ENodeTitleType::FullTitle)
+	if (Expression.IsEmpty())
 	{
-		FFormatNamedArguments Args;
-		Args.Add(TEXT("Expression"), FText::FromString(Result));
-		Result = NSLOCTEXT("K2Node", "MathExpressionSecondTitleLine", "{Expression}\nMath Expression").ToString();
+		//Result = LOCTEXT("NoExpressionTitle", "[Empty]").ToString();
 	}
 	
+	if (TitleType == ENodeTitleType::FullTitle)
+	{
+		Result = FText::Format(LOCTEXT("MathExpressionSecondTitleLine", "{0}\nMath Expression"), FText::FromString(Result)).ToString();
+	}
 	return FText::FromString(Result);
 }
 
+//------------------------------------------------------------------------------
 FString UK2Node_MathExpression::GetNodeNativeTitle(ENodeTitleType::Type TitleType) const
 {
 	FString Result = Expression;
+	if (Expression.IsEmpty())
+	{
+		//Result = TEXT("[Empty]");
+	}
 
 	if (TitleType == ENodeTitleType::FullTitle)
 	{
@@ -2317,18 +2373,21 @@ FString UK2Node_MathExpression::GetNodeNativeTitle(ENodeTitleType::Type TitleTyp
 	return Result;
 }
 
+//------------------------------------------------------------------------------
 void UK2Node_MathExpression::PostPlacedNewNode()
 {
 	Super::PostPlacedNewNode();
-
 	FEdGraphUtilities::RenameGraphToNameOrCloseToName(BoundGraph, "MathExpression");
 }
 
+//------------------------------------------------------------------------------
 void UK2Node_MathExpression::ReconstructNode()
 {
 	if (!HasAnyFlags(RF_NeedLoad|RF_NeedPostLoad))
 	{
-		RebuildExpression();
+		RebuildExpression(Expression);
 	}
 	Super::ReconstructNode();
 }
+
+#undef LOCTEXT_NAMESPACE
diff --git a/Engine/Source/Editor/BlueprintGraph/Classes/EdGraphSchema_K2.h b/Engine/Source/Editor/BlueprintGraph/Classes/EdGraphSchema_K2.h
--- a/Engine/Source/Editor/BlueprintGraph/Classes/EdGraphSchema_K2.h
+++ b/Engine/Source/Editor/BlueprintGraph/Classes/EdGraphSchema_K2.h
@@ -77,6 +77,9 @@ public:
 
 	// [FunctionMetadata] Indicates that the function should be drawn as a compact node with the specified body title
 	static const FName MD_CompactNodeTitle;
+	
+	// [FunctionMetadata] Indicates that the function should be drawn with this title over the function name
+	static const FName MD_FriendlyName;
 
 	//    property metadata
 
diff --git a/Engine/Source/Editor/BlueprintGraph/Private/EdGraphSchema_K2.cpp b/Engine/Source/Editor/BlueprintGraph/Private/EdGraphSchema_K2.cpp
--- a/Engine/Source/Editor/BlueprintGraph/Private/EdGraphSchema_K2.cpp
+++ b/Engine/Source/Editor/BlueprintGraph/Private/EdGraphSchema_K2.cpp
@@ -39,6 +39,7 @@ const FName FBlueprintMetadata::MD_FunctionCategory(TEXT("Category"));
 const FName FBlueprintMetadata::MD_DeprecatedFunction(TEXT("DeprecatedFunction"));
 const FName FBlueprintMetadata::MD_DeprecationMessage(TEXT("DeprecationMessage"));
 const FName FBlueprintMetadata::MD_CompactNodeTitle(TEXT("CompactNodeTitle"));
+const FName FBlueprintMetadata::MD_FriendlyName(TEXT("FriendlyName"));
 
 const FName FBlueprintMetadata::MD_ExposeOnSpawn(TEXT("ExposeOnSpawn"));
 const FName FBlueprintMetadata::MD_DefaultToSelf(TEXT("DefaultToSelf"));
diff --git a/Engine/Source/Editor/UnrealEd/Private/Kismet2/CompilerResultsLog.cpp b/Engine/Source/Editor/UnrealEd/Private/Kismet2/CompilerResultsLog.cpp
--- a/Engine/Source/Editor/UnrealEd/Private/Kismet2/CompilerResultsLog.cpp
+++ b/Engine/Source/Editor/UnrealEd/Private/Kismet2/CompilerResultsLog.cpp
@@ -20,19 +20,19 @@ const FName FCompilerResultsLog::Name(TEXT("CompilerResultsLog"));
 void FBacktrackMap::NotifyIntermediateObjectCreation(UObject* NewObject, UObject* SourceObject)
 {
 	// Chase the source to make sure it's really a top-level ('source code') node
-	while (UObject** SourceOfSource = SourceBacktrackMap.Find(SourceObject))
+	while (UObject** SourceOfSource = SourceBacktrackMap.Find(Cast<UObject const>(SourceObject)))
 	{
 		SourceObject = *SourceOfSource;
 	}
 
 	// Record the backtrack link
-	SourceBacktrackMap.Add(NewObject, SourceObject);
+	SourceBacktrackMap.Add(Cast<UObject const>(NewObject), SourceObject);
 }
 
 /** Returns the true source object for the passed in object */
 UObject* FBacktrackMap::FindSourceObject(UObject* PossiblyDuplicatedObject)
 {
-	UObject** RemappedIfExisting = SourceBacktrackMap.Find(PossiblyDuplicatedObject);
+	UObject** RemappedIfExisting = SourceBacktrackMap.Find(Cast<UObject const>(PossiblyDuplicatedObject));
 	if (RemappedIfExisting != NULL)
 	{
 		return *RemappedIfExisting;
@@ -44,6 +44,20 @@ UObject* FBacktrackMap::FindSourceObject(UObject* PossiblyDuplicatedObject)
 	}
 }
 
+UObject const* FBacktrackMap::FindSourceObject(UObject const* PossiblyDuplicatedObject)
+{
+	UObject** RemappedIfExisting = SourceBacktrackMap.Find(PossiblyDuplicatedObject);
+	if (RemappedIfExisting != NULL)
+	{
+		return Cast<UObject const>(*RemappedIfExisting);
+	}
+	else
+	{
+		// Not in the map, must be an unduplicated object
+		return PossiblyDuplicatedObject;
+	}
+}
+
 //////////////////////////////////////////////////////////////////////////
 // FCompilerResultsLog
 
@@ -89,6 +103,11 @@ UObject* FCompilerResultsLog::FindSourceObject(UObject* PossiblyDuplicatedObject
 	return SourceBacktrackMap.FindSourceObject(PossiblyDuplicatedObject);
 }
 
+UObject const* FCompilerResultsLog::FindSourceObject(UObject const* PossiblyDuplicatedObject)
+{
+	return SourceBacktrackMap.FindSourceObject(PossiblyDuplicatedObject);
+}
+
 /** Create a tokenized message record from a message containing @@ indicating where each UObject* in the ArgPtr list goes and place it in the MessageLog. */
 void FCompilerResultsLog::InternalLogMessage(const EMessageSeverity::Type& Severity, const TCHAR* Message, va_list ArgPtr)
 {
@@ -169,57 +188,62 @@ void FCompilerResultsLog::InternalLogMessage(const EMessageSeverity::Type& Sever
 	va_end(ArgPtr);
 
 	// Register node error/warning.
-	if ((OwnerNode != NULL) && bAnnotateMentionedNodes)
+	AnnotateNode(OwnerNode, Line);
+
+	if( !bSilentMode && (!bLogInfoOnly || (Severity == EMessageSeverity::Info)) )
+	{
+		if(Severity == EMessageSeverity::CriticalError || Severity == EMessageSeverity::Error)
+		{
+			UE_LOG(LogBlueprint, Error, TEXT("[compiler] %s"), *Line->ToText().ToString());
+		}
+		else if(Severity == EMessageSeverity::Warning || Severity == EMessageSeverity::PerformanceWarning)
+		{
+			UE_LOG(LogBlueprint, Warning, TEXT("[compiler] %s"), *Line->ToText().ToString());
+		}
+		else
+		{
+			UE_LOG(LogBlueprint, Log, TEXT("[compiler] %s"), *Line->ToText().ToString());
+		}
+	}
+}
+
+void FCompilerResultsLog::AnnotateNode(class UEdGraphNode* Node, TSharedRef<FTokenizedMessage> LogLine)
+{
+	if ((Node != NULL) && bAnnotateMentionedNodes)
 	{
 		// Determine if this message is the first or more important than the previous one (only showing one error/warning per node for now)
 		bool bUpdateMessage = true;
-		if (OwnerNode->bHasCompilerMessage)
+		if (Node->bHasCompilerMessage)
 		{
 			// Already has a message, see if we meet or trump the severity
-			bUpdateMessage = Severity <= OwnerNode->ErrorType;
+			bUpdateMessage = LogLine->GetSeverity() <= Node->ErrorType;
 		}
 		else
 		{
-			OwnerNode->ErrorMsg.Empty();
+			Node->ErrorMsg.Empty();
 		}
-
+		
 		// Update the message
 		if (bUpdateMessage)
 		{
-			OwnerNode->ErrorType = (int32)Severity;
-			OwnerNode->bHasCompilerMessage = true;
-
-			FText FullMessage = Line->ToText();
-
-			if (OwnerNode->ErrorMsg.IsEmpty())
+			Node->ErrorType = (int32)LogLine->GetSeverity();
+			Node->bHasCompilerMessage = true;
+			
+			FText FullMessage = LogLine->ToText();
+			
+			if (Node->ErrorMsg.IsEmpty())
 			{
-				OwnerNode->ErrorMsg = FullMessage.ToString();
+				Node->ErrorMsg = FullMessage.ToString();
 			}
 			else
 			{
 				FFormatNamedArguments Args;
-				Args.Add( TEXT("PreviousMessage"), FText::FromString( OwnerNode->ErrorMsg ) );
+				Args.Add( TEXT("PreviousMessage"), FText::FromString( Node->ErrorMsg ) );
 				Args.Add( TEXT("NewMessage"), FullMessage );
-				OwnerNode->ErrorMsg = FText::Format( LOCTEXT("AggregateMessagesFormatter", "{PreviousMessage}\n{NewMessage}"), Args ).ToString();
+				Node->ErrorMsg = FText::Format( LOCTEXT("AggregateMessagesFormatter", "{PreviousMessage}\n{NewMessage}"), Args ).ToString();
 			}
 		}
 	}
-
-	if( !bSilentMode && (!bLogInfoOnly || (Severity == EMessageSeverity::Info)) )
-	{
-		if(Severity == EMessageSeverity::CriticalError || Severity == EMessageSeverity::Error)
-		{
-			UE_LOG(LogBlueprint, Error, TEXT("[compiler] %s"), *Line->ToText().ToString());
-		}
-		else if(Severity == EMessageSeverity::Warning || Severity == EMessageSeverity::PerformanceWarning)
-		{
-			UE_LOG(LogBlueprint, Warning, TEXT("[compiler] %s"), *Line->ToText().ToString());
-		}
-		else
-		{
-			UE_LOG(LogBlueprint, Log, TEXT("[compiler] %s"), *Line->ToText().ToString());
-		}
-	}
 }
 
 TArray< TSharedRef<FTokenizedMessage> > FCompilerResultsLog::ParseCompilerLogDump(const FString& LogDump)
@@ -371,6 +395,51 @@ void FCompilerResultsLog::NoteVA(const TCHAR* Message, va_list ArgPtr)
 	InternalLogMessage(EMessageSeverity::Info, Message, ArgPtr);
 }
 
+void FCompilerResultsLog::Append(FCompilerResultsLog const& Other)
+{
+	for (TSharedRef<FTokenizedMessage> const& Message : Other.Messages)
+	{
+		switch (Message->GetSeverity())
+		{
+		case EMessageSeverity::Warning:
+			{
+				++NumWarnings;
+				break;
+			}
+				
+		case EMessageSeverity::Error:
+			{
+				++NumErrors;
+				break;
+			}
+		}
+		Messages.Add(Message);
+		
+		UEdGraphNode* OwnerNode = nullptr;
+		for (TSharedRef<IMessageToken> const& Token : Message->GetMessageTokens())
+		{
+			if (Token->GetType() != EMessageToken::Object)
+			{
+				continue;
+			}
+			
+			FWeakObjectPtr ObjectPtr = ((FUObjectToken&)Token.Get()).GetObject();
+			if (!ObjectPtr.IsValid())
+			{
+				continue;
+			}
+			UObject* ObjectArgument = ObjectPtr.Get();
+			
+			OwnerNode = Cast<UEdGraphNode>(ObjectArgument);
+			if (UEdGraphPin const* Pin = Cast<UEdGraphPin>(ObjectArgument))
+			{
+				OwnerNode = Pin->GetOwningNodeUnchecked();
+			}
+		}
+		AnnotateNode(OwnerNode, Message);
+	}
+}
+
 #undef LOCTEXT_NAMESPACE
 
 #endif	//#if !WITH_EDITOR
diff --git a/Engine/Source/Editor/UnrealEd/Public/Kismet2/CompilerResultsLog.h b/Engine/Source/Editor/UnrealEd/Public/Kismet2/CompilerResultsLog.h
--- a/Engine/Source/Editor/UnrealEd/Public/Kismet2/CompilerResultsLog.h
+++ b/Engine/Source/Editor/UnrealEd/Public/Kismet2/CompilerResultsLog.h
@@ -12,7 +12,7 @@ class UNREALED_API FBacktrackMap
 {
 protected:
 	// Maps from transient object created during compiling to original 'source code' object
-	TMap<UObject*, UObject*> SourceBacktrackMap;
+	TMap<UObject const*, UObject*> SourceBacktrackMap;
 public:
 	FBacktrackMap(){}
 	virtual ~FBacktrackMap(){}
@@ -22,6 +22,7 @@ public:
 
 	/** Returns the true source object for the passed in object */
 	UObject* FindSourceObject(UObject* PossiblyDuplicatedObject);
+	UObject const* FindSourceObject(UObject const* PossiblyDuplicatedObject);
 };
 
 /** This class represents a log of compiler output lines (errors, warnings, and information notes), each of which can be a rich tokenized message */
@@ -74,6 +75,7 @@ public:
 
 	/** Returns the true source object for the passed in object */
 	UObject* FindSourceObject(UObject* PossiblyDuplicatedObject);
+	UObject const* FindSourceObject(UObject const* PossiblyDuplicatedObject);
 
 	/** Returns the true source object for the passed in object; does type checking on the result */
 	template <typename T>
@@ -81,12 +83,23 @@ public:
 	{
 		return CastChecked<T>(FindSourceObject(PossiblyDuplicatedObject));
 	}
+	
+	template <typename T>
+	T const* FindSourceObjectTypeChecked(UObject const* PossiblyDuplicatedObject)
+	{
+		return CastChecked<T const>(FindSourceObject(PossiblyDuplicatedObject));
+	}
+	
+	void Append(FCompilerResultsLog const& Other);
 protected:
 	/** Create a tokenized message record from a message containing @@ indicating where each UObject* in the ArgPtr list goes and place it in the MessageLog. */
 	void InternalLogMessage(const EMessageSeverity::Type& Severity, const TCHAR* Message, va_list ArgPtr);
 
 	/** Callback when a token is activated */
 	void OnTokenActivated(const TSharedRef<class IMessageToken>& InTokenRef);
+	
+	/** */
+	void AnnotateNode(class UEdGraphNode* Node, TSharedRef<FTokenizedMessage> LogLine);
 
 private:
 	/** Parses a compiler log dump to generate tokenized output */
